MotionDesign : PropertyAnimator
- Adding new curve wave animator that supports custom wave and ease curves
- Adding In/Out/InOut Back, Circular, Cubic, Exponential, Overshoot, Quadratic, Quartic, Quintic, Sine ease curve assets
- Adding Bounce, Constant, Cosine, Elastic, Sawtooth, Sine, Square, Triangle wave curve assets

#virtualized

diff --git a/Engine/Plugins/Experimental/PropertyAnimator/Source/PropertyAnimator/Private/Animators/PropertyAnimatorClock.cpp b/Engine/Plugins/Experimental/PropertyAnimator/Source/PropertyAnimator/Private/Animators/PropertyAnimatorClock.cpp
--- a/Engine/Plugins/Experimental/PropertyAnimator/Source/PropertyAnimator/Private/Animators/PropertyAnimatorClock.cpp
+++ b/Engine/Plugins/Experimental/PropertyAnimator/Source/PropertyAnimator/Private/Animators/PropertyAnimatorClock.cpp
@@ -48,9 +48,11 @@ void UPropertyAnimatorClock::EvaluateProperties(FInstancedPropertyBag& InParamet
 	const FString FormattedDateTime = DateTime.ToFormattedString(*DisplayFormat);
 
 	EvaluateEachLinkedProperty<UPropertyAnimatorCoreContext>([this, FormattedDateTime](
-		UPropertyAnimatorCoreContext* InContext,
-		const FPropertyAnimatorCoreData& InResolvedProperty,
-		FInstancedPropertyBag& InEvaluatedValues)->bool
+		UPropertyAnimatorCoreContext* InContext
+		, const FPropertyAnimatorCoreData& InResolvedProperty
+		, FInstancedPropertyBag& InEvaluatedValues
+		, int32 InRangeIndex
+		, int32 InRangeMax)->bool
 	{
 		const FName DisplayName(InResolvedProperty.GetPathHash());
 
diff --git a/Engine/Plugins/Experimental/PropertyAnimator/Source/PropertyAnimator/Private/Animators/PropertyAnimatorCurveWave.cpp b/Engine/Plugins/Experimental/PropertyAnimator/Source/PropertyAnimator/Private/Animators/PropertyAnimatorCurveWave.cpp
new file mode 100644
--- /dev/null
+++ b/Engine/Plugins/Experimental/PropertyAnimator/Source/PropertyAnimator/Private/Animators/PropertyAnimatorCurveWave.cpp
@@ -0,0 +1,110 @@
+// Copyright Epic Games, Inc. All Rights Reserved.
+
+#include "Animators/PropertyAnimatorCurveWave.h"
+
+#include "Curves/PropertyAnimatorEaseCurve.h"
+#include "Curves/PropertyAnimatorWaveCurve.h"
+#include "Properties/PropertyAnimatorFloatContext.h"
+
+UPropertyAnimatorCurveWave::UPropertyAnimatorCurveWave()
+{
+	SetAnimatorDisplayName(DefaultAnimatorName);
+}
+
+void UPropertyAnimatorCurveWave::SetWaveCurve(UPropertyAnimatorWaveCurve* InCurve)
+{
+	WaveCurve = InCurve;
+}
+
+void UPropertyAnimatorCurveWave::SetEaseIn(const FPropertyAnimatorCurveEasing& InEasing)
+{
+	EaseIn = InEasing;
+}
+
+void UPropertyAnimatorCurveWave::SetEaseOut(const FPropertyAnimatorCurveEasing& InEasing)
+{
+	EaseOut = InEasing;
+}
+
+#if WITH_EDITOR
+void UPropertyAnimatorCurveWave::PostEditChangeProperty(FPropertyChangedEvent& InPropertyChangedEvent)
+{
+	Super::PostEditChangeProperty(InPropertyChangedEvent);
+
+	const FName MemberName = InPropertyChangedEvent.GetMemberPropertyName();
+
+	if (MemberName == GET_MEMBER_NAME_CHECKED(UPropertyAnimatorCurveWave, EaseIn))
+	{
+		OnEaseInChanged();
+	}
+	else if (MemberName == GET_MEMBER_NAME_CHECKED(UPropertyAnimatorCurveWave, EaseOut))
+	{
+		OnEaseOutChanged();
+	}
+}
+#endif
+
+bool UPropertyAnimatorCurveWave::EvaluateProperty(const FPropertyAnimatorCoreData& InPropertyData, UPropertyAnimatorCoreContext* InContext, FInstancedPropertyBag& InParameters, FInstancedPropertyBag& OutEvaluationResult) const
+{
+	if (!WaveCurve)
+	{
+		return false;
+	}
+
+	const FRichCurve& SampleCurve = WaveCurve->FloatCurve;
+
+	float MinTime = 0;
+	float MaxTime = 0;
+	SampleCurve.GetTimeRange(MinTime, MaxTime);
+
+	float MinValue = 0;
+	float MaxValue = 0;
+	SampleCurve.GetValueRange(MinValue, MaxValue);
+
+	const double TimeElapsed = InParameters.GetValueDouble(TimeElapsedParameterName).GetValue();
+	const float Frequency = InParameters.GetValueDouble(FrequencyParameterName).GetValue();
+	const float Period = 1.f / Frequency;
+
+	const float SampleTime = FMath::Fmod(TimeElapsed, Period);
+	const float NormalizedSampleTime = FMath::GetMappedRangeValueClamped(FVector2D(0, Period), FVector2D(MinTime, MaxTime), SampleTime);
+
+	const float SampleValue = SampleCurve.Eval(NormalizedSampleTime);
+	float SampleValueNormalized = FMath::GetMappedRangeValueClamped(FVector2D(MinValue, MaxValue), FVector2D(0, 1), SampleValue);
+
+	if (EaseIn.EaseCurve && SampleTime < EaseIn.EaseDuration)
+	{
+		const FRichCurve& EaseCurve = EaseIn.EaseCurve->FloatCurve;
+		const float EaseTimeNormalized = FMath::GetMappedRangeValueClamped(FVector2D(0, EaseIn.EaseDuration), FVector2D(0, 1), SampleTime);
+		SampleValueNormalized *= EaseCurve.Eval(EaseTimeNormalized);
+	}
+
+	if (EaseOut.EaseCurve && SampleTime > (CycleDuration - EaseOut.EaseDuration))
+	{
+		const FRichCurve& EaseCurve = EaseOut.EaseCurve->FloatCurve;
+		const float EaseTimeNormalized = 1 - FMath::GetMappedRangeValueClamped(FVector2D(CycleDuration - EaseOut.EaseDuration, CycleDuration), FVector2D(0, 1), SampleTime);
+		SampleValueNormalized *= EaseCurve.Eval(EaseTimeNormalized);
+	}
+
+	InParameters.AddProperty(AlphaParameterName, EPropertyBagPropertyType::Float);
+	InParameters.SetValueFloat(AlphaParameterName, SampleValueNormalized);
+
+	return InContext->EvaluateProperty(InPropertyData, InParameters, OutEvaluationResult);
+}
+
+void UPropertyAnimatorCurveWave::OnEaseInChanged()
+{
+	EaseIn.EaseDuration = FMath::Clamp(EaseIn.EaseDuration, 0, CycleDuration - EaseOut.EaseDuration);
+}
+
+void UPropertyAnimatorCurveWave::OnEaseOutChanged()
+{
+	EaseOut.EaseDuration = FMath::Clamp(EaseOut.EaseDuration, 0, CycleDuration - EaseIn.EaseDuration);
+}
+
+void UPropertyAnimatorCurveWave::OnCycleDurationChanged()
+{
+	Super::OnCycleDurationChanged();
+
+	OnEaseInChanged();
+	OnEaseOutChanged();
+}
diff --git a/Engine/Plugins/Experimental/PropertyAnimator/Source/PropertyAnimator/Private/Animators/PropertyAnimatorFloatBase.cpp b/Engine/Plugins/Experimental/PropertyAnimator/Source/PropertyAnimator/Private/Animators/PropertyAnimatorFloatBase.cpp
--- a/Engine/Plugins/Experimental/PropertyAnimator/Source/PropertyAnimator/Private/Animators/PropertyAnimatorFloatBase.cpp
+++ b/Engine/Plugins/Experimental/PropertyAnimator/Source/PropertyAnimator/Private/Animators/PropertyAnimatorFloatBase.cpp
@@ -156,49 +156,63 @@ void UPropertyAnimatorFloatBase::EvaluateProperties(FInstancedPropertyBag& InPar
 	double TimeElapsed = InParameters.GetValueDouble(TimeElapsedParameterName).GetValue();
 	RandomStream = FRandomStream(Seed);
 
-	if (CycleMode == EPropertyAnimatorCycleMode::DoOnce)
+	EvaluateEachLinkedProperty<UPropertyAnimatorCoreContext>([this, &TimeElapsed, &AnimatorMagnitude, &InParameters](
+		UPropertyAnimatorCoreContext* InOptions
+		, const FPropertyAnimatorCoreData& InResolvedProperty
+		, FInstancedPropertyBag& InEvaluatedValues
+		, int32 InRangeIndex
+		, int32 InRangeMax)->bool
 	{
-		if (FMath::Abs(TimeElapsed) > CycleDuration)
+		const double RandomTimeOffset = bRandomTimeOffset ? RandomStream.GetFraction() : 0;
+
+		const double AbsTimeOffset = FMath::Abs(TimeOffset);
+		const double MaxTimeOffset = InRangeMax * AbsTimeOffset;
+		double PropertyTimeElapsed = TimeElapsed + RandomTimeOffset;
+
+		if (TimeOffset >= 0)
 		{
-			return;
+			PropertyTimeElapsed += InRangeIndex * AbsTimeOffset;
 		}
-	}
-	else if (CycleMode == EPropertyAnimatorCycleMode::Loop)
-	{
-		TimeElapsed = FMath::Fmod(TimeElapsed, CycleDuration + CycleGapDuration);
-
-		if (TimeElapsed > CycleDuration)
+		else
 		{
-			TimeElapsed = CycleDuration - UE_KINDA_SMALL_NUMBER;
+			PropertyTimeElapsed += MaxTimeOffset - InRangeIndex * AbsTimeOffset;
 		}
-	}
-	else if (CycleMode == EPropertyAnimatorCycleMode::PingPong)
-	{
-		const bool bReverse = FMath::Modulo(FMath::TruncToInt32(TimeElapsed / (CycleDuration + CycleGapDuration)), 2) != 0;
-		TimeElapsed = FMath::Fmod(TimeElapsed, CycleDuration + CycleGapDuration);
 
-		if (TimeElapsed > CycleDuration)
+		if (CycleMode == EPropertyAnimatorCycleMode::DoOnce)
 		{
-			TimeElapsed = CycleDuration - UE_KINDA_SMALL_NUMBER;
+			if (FMath::Abs(PropertyTimeElapsed) > CycleDuration)
+			{
+				return false;
+			}
 		}
-
-		if (bReverse)
+		else if (CycleMode == EPropertyAnimatorCycleMode::Loop)
 		{
-			TimeElapsed = CycleDuration - FMath::Fmod(TimeElapsed, CycleDuration);
+			PropertyTimeElapsed = FMath::Fmod(PropertyTimeElapsed, CycleDuration + MaxTimeOffset + CycleGapDuration);
+
+			if (FMath::Abs(PropertyTimeElapsed) > CycleDuration)
+			{
+				PropertyTimeElapsed = CycleDuration - UE_KINDA_SMALL_NUMBER;
+			}
 		}
-		else
+		else if (CycleMode == EPropertyAnimatorCycleMode::PingPong)
 		{
-			TimeElapsed = FMath::Fmod(TimeElapsed, CycleDuration);
+			const bool bReverse = FMath::Modulo(FMath::TruncToInt32(PropertyTimeElapsed / (CycleDuration + MaxTimeOffset + CycleGapDuration)), 2) != 0;
+			PropertyTimeElapsed = FMath::Fmod(PropertyTimeElapsed, CycleDuration + MaxTimeOffset + CycleGapDuration);
+
+			if (FMath::Abs(PropertyTimeElapsed) > CycleDuration)
+			{
+				PropertyTimeElapsed = CycleDuration - UE_KINDA_SMALL_NUMBER;
+			}
+
+			if (bReverse)
+			{
+				PropertyTimeElapsed = CycleDuration - FMath::Fmod(PropertyTimeElapsed, CycleDuration);
+			}
+			else
+			{
+				PropertyTimeElapsed = FMath::Fmod(PropertyTimeElapsed, CycleDuration);
+			}
 		}
-	}
-
-	EvaluateEachLinkedProperty<UPropertyAnimatorCoreContext>([this, &TimeElapsed, &AnimatorMagnitude, &InParameters](
-		UPropertyAnimatorCoreContext* InOptions
-		, const FPropertyAnimatorCoreData& InResolvedProperty
-		, FInstancedPropertyBag& InEvaluatedValues)->bool
-	{
-		const double RandomTimeOffset = bRandomTimeOffset ? RandomStream.GetFraction() : 0;
-		TimeElapsed += TimeOffset + RandomTimeOffset;
 
 		if (Magnitude != 0
 			&& CycleDuration > 0
@@ -209,7 +223,7 @@ void UPropertyAnimatorFloatBase::EvaluateProperties(FInstancedPropertyBag& InPar
 			InParameters.SetValueFloat(FrequencyParameterName, 1.f / CycleDuration);
 
 			// Time Elapsed
-			InParameters.SetValueDouble(TimeElapsedParameterName, TimeElapsed + InOptions->GetTimeOffset());
+			InParameters.SetValueDouble(TimeElapsedParameterName, PropertyTimeElapsed);
 
 			// Magnitude
 			InParameters.SetValueFloat(MagnitudeParameterName, AnimatorMagnitude * InOptions->GetMagnitude());
diff --git a/Engine/Plugins/Experimental/PropertyAnimator/Source/PropertyAnimator/Public/Animators/PropertyAnimatorCurveWave.h b/Engine/Plugins/Experimental/PropertyAnimator/Source/PropertyAnimator/Public/Animators/PropertyAnimatorCurveWave.h
new file mode 100644
--- /dev/null
+++ b/Engine/Plugins/Experimental/PropertyAnimator/Source/PropertyAnimator/Public/Animators/PropertyAnimatorCurveWave.h
@@ -0,0 +1,83 @@
+// Copyright Epic Games, Inc. All Rights Reserved.
+
+#pragma once
+
+#include "Animators/PropertyAnimatorFloatBase.h"
+#include "PropertyAnimatorCurveWave.generated.h"
+
+class UPropertyAnimatorEaseCurve;
+class UPropertyAnimatorWaveCurve;
+
+USTRUCT(BlueprintType)
+struct FPropertyAnimatorCurveEasing
+{
+	GENERATED_BODY()
+
+	UPROPERTY(EditInstanceOnly, Category="Animator")
+	TObjectPtr<UPropertyAnimatorEaseCurve> EaseCurve;
+
+	UPROPERTY(EditInstanceOnly, Interp, Category="Animator", meta=(ClampMin="0", Units=Seconds))
+	float EaseDuration = 0.f;
+};
+
+/**
+ * Applies a wave movement from a curve on supported float properties
+ */
+UCLASS(MinimalAPI, AutoExpandCategories=("Animator"))
+class UPropertyAnimatorCurveWave : public UPropertyAnimatorFloatBase
+{
+	GENERATED_BODY()
+
+public:
+	static constexpr const TCHAR* DefaultAnimatorName = TEXT("CurveWave");
+
+	UPropertyAnimatorCurveWave();
+
+	void SetWaveCurve(UPropertyAnimatorWaveCurve* InCurve);
+
+	UPropertyAnimatorWaveCurve* GetWaveCurve() const
+	{
+		return WaveCurve;
+	}
+
+	void SetEaseIn(const FPropertyAnimatorCurveEasing& InEasing);
+
+	const FPropertyAnimatorCurveEasing& GetEaseIn() const
+	{
+		return EaseIn;
+	}
+
+	void SetEaseOut(const FPropertyAnimatorCurveEasing& InEasing);
+
+	const FPropertyAnimatorCurveEasing& GetEaseOut() const
+	{
+		return EaseOut;
+	}
+
+protected:
+	//~ Begin UObject
+#if WITH_EDITOR
+	virtual void PostEditChangeProperty(FPropertyChangedEvent& InPropertyChangedEvent) override;
+#endif
+	//~ End UObject
+
+	//~ Begin UPropertyAnimatorFloatBase
+	virtual bool EvaluateProperty(const FPropertyAnimatorCoreData& InPropertyData, UPropertyAnimatorCoreContext* InContext, FInstancedPropertyBag& InParameters, FInstancedPropertyBag& OutEvaluationResult) const override;
+	//~ End UPropertyAnimatorFloatBase
+
+	void OnEaseInChanged();
+	void OnEaseOutChanged();
+	virtual void OnCycleDurationChanged() override;
+
+	/** The wave curve to sample for the animation */
+	UPROPERTY(EditInstanceOnly, Setter, Getter, Category="Animator")
+	TObjectPtr<UPropertyAnimatorWaveCurve> WaveCurve;
+
+	/** Ease in for this effect */
+	UPROPERTY(EditInstanceOnly, Setter, Getter, Category="Animator")
+	FPropertyAnimatorCurveEasing EaseIn;
+
+	/** Ease out for this effect */
+	UPROPERTY(EditInstanceOnly, Setter, Getter, Category="Animator")
+    FPropertyAnimatorCurveEasing EaseOut;
+};
\ No newline at end of file
diff --git a/Engine/Plugins/Experimental/PropertyAnimator/Source/PropertyAnimator/Public/Animators/PropertyAnimatorFloatBase.h b/Engine/Plugins/Experimental/PropertyAnimator/Source/PropertyAnimator/Public/Animators/PropertyAnimatorFloatBase.h
--- a/Engine/Plugins/Experimental/PropertyAnimator/Source/PropertyAnimator/Public/Animators/PropertyAnimatorFloatBase.h
+++ b/Engine/Plugins/Experimental/PropertyAnimator/Source/PropertyAnimator/Public/Animators/PropertyAnimatorFloatBase.h
@@ -117,7 +117,7 @@ protected:
 
 	/** Time offset accumulated for each property for every round */
 	UPROPERTY(EditInstanceOnly, Setter, Getter, Category="Animator", meta=(Units=Seconds))
-	double TimeOffset = 0.f;
+	double TimeOffset = 0;
 
 private:
 	/** Random stream for time offset */
diff --git a/Engine/Plugins/Experimental/PropertyAnimator/Source/PropertyAnimator/Public/Curves/PropertyAnimatorEaseCurve.h b/Engine/Plugins/Experimental/PropertyAnimator/Source/PropertyAnimator/Public/Curves/PropertyAnimatorEaseCurve.h
new file mode 100644
--- /dev/null
+++ b/Engine/Plugins/Experimental/PropertyAnimator/Source/PropertyAnimator/Public/Curves/PropertyAnimatorEaseCurve.h
@@ -0,0 +1,13 @@
+// Copyright Epic Games, Inc. All Rights Reserved.
+
+#pragma once
+
+#include "Curves/CurveFloat.h"
+#include "PropertyAnimatorEaseCurve.generated.h"
+
+/** Base class for property animator easing curves (0-1) */
+UCLASS(MinimalAPI)
+class UPropertyAnimatorEaseCurve : public UCurveFloat
+{
+	GENERATED_BODY()
+};
diff --git a/Engine/Plugins/Experimental/PropertyAnimator/Source/PropertyAnimator/Public/Curves/PropertyAnimatorWaveCurve.h b/Engine/Plugins/Experimental/PropertyAnimator/Source/PropertyAnimator/Public/Curves/PropertyAnimatorWaveCurve.h
new file mode 100644
--- /dev/null
+++ b/Engine/Plugins/Experimental/PropertyAnimator/Source/PropertyAnimator/Public/Curves/PropertyAnimatorWaveCurve.h
@@ -0,0 +1,13 @@
+// Copyright Epic Games, Inc. All Rights Reserved.
+
+#pragma once
+
+#include "Curves/CurveFloat.h"
+#include "PropertyAnimatorWaveCurve.generated.h"
+
+/** Base class for property animator wave curves (inf-inf) */
+UCLASS(MinimalAPI)
+class UPropertyAnimatorWaveCurve : public UCurveFloat
+{
+	GENERATED_BODY()
+};
diff --git a/Engine/Plugins/Experimental/PropertyAnimator/Source/PropertyAnimatorEditor/Private/Styles/PropertyAnimatorEditorStyle.cpp b/Engine/Plugins/Experimental/PropertyAnimator/Source/PropertyAnimatorEditor/Private/Styles/PropertyAnimatorEditorStyle.cpp
--- a/Engine/Plugins/Experimental/PropertyAnimator/Source/PropertyAnimatorEditor/Private/Styles/PropertyAnimatorEditorStyle.cpp
+++ b/Engine/Plugins/Experimental/PropertyAnimator/Source/PropertyAnimatorEditor/Private/Styles/PropertyAnimatorEditorStyle.cpp
@@ -20,12 +20,10 @@ FPropertyAnimatorEditorStyle::FPropertyAnimatorEditorStyle()
 
 	ContentRootDir = FPaths::Combine(Plugin->GetBaseDir(), TEXT("Resources"));
 
-	Set("ClassIcon.PropertyAnimatorBounce",     new IMAGE_BRUSH_SVG("Animators/Bounce",    Icon16x16));
-	Set("ClassIcon.PropertyAnimatorOscillate",  new IMAGE_BRUSH_SVG("Animators/Oscillate", Icon16x16));
+	Set("ClassIcon.PropertyAnimatorCurveWave",  new IMAGE_BRUSH_SVG("Animators/Bounce",    Icon16x16));
+	Set("ClassIcon.PropertyAnimatorSoundWave",  new IMAGE_BRUSH("Animators/SoundWave",     Icon16x16));
 	Set("ClassIcon.PropertyAnimatorTime",       new IMAGE_BRUSH_SVG("Animators/Time",      Icon16x16));
 	Set("ClassIcon.PropertyAnimatorWiggle",     new IMAGE_BRUSH_SVG("Animators/Wiggle",    Icon16x16));
-	Set("ClassIcon.PropertyAnimatorBlink",      new IMAGE_BRUSH_SVG("Animators/Blink",     Icon16x16));
-	Set("ClassIcon.PropertyAnimatorSoundWave",  new IMAGE_BRUSH("Animators/SoundWave",     Icon16x16));
 
 	FSlateStyleRegistry::RegisterSlateStyle(*this);
 }
diff --git a/Engine/Plugins/Experimental/PropertyAnimatorCore/Source/PropertyAnimatorCore/Private/Properties/PropertyAnimatorCoreContext.cpp b/Engine/Plugins/Experimental/PropertyAnimatorCore/Source/PropertyAnimatorCore/Private/Properties/PropertyAnimatorCoreContext.cpp
--- a/Engine/Plugins/Experimental/PropertyAnimatorCore/Source/PropertyAnimatorCore/Private/Properties/PropertyAnimatorCoreContext.cpp
+++ b/Engine/Plugins/Experimental/PropertyAnimatorCore/Source/PropertyAnimatorCore/Private/Properties/PropertyAnimatorCoreContext.cpp
@@ -90,11 +90,6 @@ void UPropertyAnimatorCoreContext::SetMagnitude(float InMagnitude)
 	Magnitude = FMath::Clamp(InMagnitude, 0.f, 1.f);
 }
 
-void UPropertyAnimatorCoreContext::SetTimeOffset(double InTimeOffset)
-{
-	TimeOffset = InTimeOffset;
-}
-
 void UPropertyAnimatorCoreContext::SetMode(EPropertyAnimatorCoreMode InMode)
 {
 	if (InMode == Mode)
diff --git a/Engine/Plugins/Experimental/PropertyAnimatorCore/Source/PropertyAnimatorCore/Public/Animators/PropertyAnimatorCoreBase.h b/Engine/Plugins/Experimental/PropertyAnimatorCore/Source/PropertyAnimatorCore/Public/Animators/PropertyAnimatorCoreBase.h
--- a/Engine/Plugins/Experimental/PropertyAnimatorCore/Source/PropertyAnimatorCore/Public/Animators/PropertyAnimatorCoreBase.h
+++ b/Engine/Plugins/Experimental/PropertyAnimatorCore/Source/PropertyAnimatorCore/Public/Animators/PropertyAnimatorCoreBase.h
@@ -206,7 +206,9 @@ protected:
 		TFunctionRef<bool(
 			InContextClass* /** InPropertyContext */
 			, const FPropertyAnimatorCoreData& /** InResolvedProperty */
-			, FInstancedPropertyBag& /** OutEvaluation */)> InFunction
+			, FInstancedPropertyBag& /** OutEvaluation */
+			, int32 InRangeIndex
+			, int32 InRangeMax)> InFunction
 		)
 	{
 		checkf(bEvaluatingProperties, TEXT("EvaluateEachLinkedProperty can only be called in EvaluateProperties"))
@@ -220,14 +222,18 @@ protected:
 					continue;
 				}
 
-				for (const FPropertyAnimatorCoreData& ResolvedPropertyData : PropertyContext->ResolveProperty(true))
+				const TArray<FPropertyAnimatorCoreData> ResolvedProperties = PropertyContext->ResolveProperty(true);
+
+				for (int32 Index = 0; Index < ResolvedProperties.Num(); Index++)
 				{
+					const FPropertyAnimatorCoreData& ResolvedPropertyData = ResolvedProperties[Index];
+
 					if (!ResolvedPropertyData.IsResolved())
 					{
 						continue;
 					}
 
-					if (InFunction(PropertyContext, ResolvedPropertyData, EvaluatedPropertyValues))
+					if (InFunction(PropertyContext, ResolvedPropertyData, EvaluatedPropertyValues, Index, ResolvedProperties.Num() - 1))
 					{
 						PropertyContext->CommitEvaluationResult(ResolvedPropertyData, EvaluatedPropertyValues);
 					}
diff --git a/Engine/Plugins/Experimental/PropertyAnimatorCore/Source/PropertyAnimatorCore/Public/Properties/PropertyAnimatorCoreContext.h b/Engine/Plugins/Experimental/PropertyAnimatorCore/Source/PropertyAnimatorCore/Public/Properties/PropertyAnimatorCoreContext.h
--- a/Engine/Plugins/Experimental/PropertyAnimatorCore/Source/PropertyAnimatorCore/Public/Properties/PropertyAnimatorCoreContext.h
+++ b/Engine/Plugins/Experimental/PropertyAnimatorCore/Source/PropertyAnimatorCore/Public/Properties/PropertyAnimatorCoreContext.h
@@ -59,12 +59,6 @@ public:
 		return Magnitude;
 	}
 
-	PROPERTYANIMATORCORE_API void SetTimeOffset(double InTimeOffset);
-	double GetTimeOffset() const
-	{
-		return TimeOffset;
-	}
-
 	PROPERTYANIMATORCORE_API void SetMode(EPropertyAnimatorCoreMode InMode);
 	EPropertyAnimatorCoreMode GetMode() const
 	{
@@ -157,10 +151,6 @@ private:
 	UPROPERTY(EditInstanceOnly, Setter, Getter, Category="Animator", meta=(ClampMin="0.0", ClampMax="1.0"))
 	float Magnitude = 1.f;
 
-	/** Time offset variation for evaluation */
-	UPROPERTY(EditInstanceOnly, Setter, Getter, Category="Animator")
-	double TimeOffset = 0.f;
-
 	/** Edit condition for modes */
 	UPROPERTY(Transient)
 	bool bEditMode = true;
