MotionDesign : PropertyAnimator
- Updating property metadata for UI

diff --git a/Engine/Plugins/Experimental/PropertyAnimator/Source/PropertyAnimator/Private/Animators/PropertyAnimatorFloatBase.cpp b/Engine/Plugins/Experimental/PropertyAnimator/Source/PropertyAnimator/Private/Animators/PropertyAnimatorFloatBase.cpp
--- a/Engine/Plugins/Experimental/PropertyAnimator/Source/PropertyAnimator/Private/Animators/PropertyAnimatorFloatBase.cpp
+++ b/Engine/Plugins/Experimental/PropertyAnimator/Source/PropertyAnimator/Private/Animators/PropertyAnimatorFloatBase.cpp
@@ -167,7 +167,7 @@ void UPropertyAnimatorFloatBase::EvaluateProperties(FInstancedPropertyBag& InPar
 
 		const double AbsTimeOffset = FMath::Abs(TimeOffset);
 		const double MaxTimeOffset = InRangeMax * AbsTimeOffset;
-		double PropertyTimeElapsed = TimeElapsed + RandomTimeOffset;
+		double PropertyTimeElapsed = TimeElapsed - MaxTimeOffset + RandomTimeOffset;
 
 		if (TimeOffset >= 0)
 		{
diff --git a/Engine/Plugins/Experimental/PropertyAnimator/Source/PropertyAnimator/Public/Animators/PropertyAnimatorFloatBase.h b/Engine/Plugins/Experimental/PropertyAnimator/Source/PropertyAnimator/Public/Animators/PropertyAnimatorFloatBase.h
--- a/Engine/Plugins/Experimental/PropertyAnimator/Source/PropertyAnimator/Public/Animators/PropertyAnimatorFloatBase.h
+++ b/Engine/Plugins/Experimental/PropertyAnimator/Source/PropertyAnimator/Public/Animators/PropertyAnimatorFloatBase.h
@@ -104,7 +104,7 @@ protected:
 	EPropertyAnimatorCycleMode CycleMode = EPropertyAnimatorCycleMode::Loop;
 
 	/** Time gap between each cycle */
-	UPROPERTY(EditInstanceOnly, Category="Animator", meta=(ClampMin="0", Units=Seconds))
+	UPROPERTY(EditInstanceOnly, Category="Animator", meta=(ClampMin="0", Units=Seconds, EditCondition="CycleMode != EPropertyAnimatorCycleMode::DoOnce", EditConditionHides))
 	float CycleGapDuration = 0.f;
 
 	/** Use random time offset to add variation in animation */
diff --git a/Engine/Plugins/Experimental/PropertyAnimator/Source/PropertyAnimator/Public/Properties/PropertyAnimatorVectorContext.h b/Engine/Plugins/Experimental/PropertyAnimator/Source/PropertyAnimator/Public/Properties/PropertyAnimatorVectorContext.h
--- a/Engine/Plugins/Experimental/PropertyAnimator/Source/PropertyAnimator/Public/Properties/PropertyAnimatorVectorContext.h
+++ b/Engine/Plugins/Experimental/PropertyAnimator/Source/PropertyAnimator/Public/Properties/PropertyAnimatorVectorContext.h
@@ -39,7 +39,7 @@ protected:
 	FVector GetClampedAmplitude(FVector InAmplitude);
 
 	/** The minimum value should be remapped to that values */
-	UPROPERTY(EditInstanceOnly, Setter, Getter, Category="Animator")
+	UPROPERTY(EditInstanceOnly, Setter, Getter, Category="Animator", meta=(MotionDesignVectorWidget, AllowPreserveRatio="XYZ"))
 	FVector AmplitudeMin = FVector::ZeroVector;
 
 	/** Some properties are clamped and cannot go past a specific min value */
@@ -47,7 +47,7 @@ protected:
 	TOptional<FVector> AmplitudeClampMin;
 
 	/** The maximum value should be remapped to that values */
-	UPROPERTY(EditInstanceOnly, Setter, Getter, Category="Animator")
+	UPROPERTY(EditInstanceOnly, Setter, Getter, Category="Animator", meta=(MotionDesignVectorWidget, AllowPreserveRatio="XYZ"))
 	FVector AmplitudeMax = FVector::ZeroVector;
 
 	/** Some properties are clamped and cannot go past a specific min value */
diff --git a/Engine/Plugins/Experimental/PropertyAnimatorCore/Source/PropertyAnimatorCore/Public/TimeSources/PropertyAnimatorCoreMachineClockTimeSource.h b/Engine/Plugins/Experimental/PropertyAnimatorCore/Source/PropertyAnimatorCore/Public/TimeSources/PropertyAnimatorCoreMachineClockTimeSource.h
--- a/Engine/Plugins/Experimental/PropertyAnimatorCore/Source/PropertyAnimatorCore/Public/TimeSources/PropertyAnimatorCoreMachineClockTimeSource.h
+++ b/Engine/Plugins/Experimental/PropertyAnimatorCore/Source/PropertyAnimatorCore/Public/TimeSources/PropertyAnimatorCoreMachineClockTimeSource.h
@@ -80,7 +80,7 @@ protected:
 	* 1h = 1 hour
 	* 120s = 2 minutes
 	*/
-	UPROPERTY(EditInstanceOnly, Setter, Getter, Category="Animator", meta=(EditCondition="Mode == EPropertyAnimatorCoreLocalTimeMode::Countdown", EditConditionHides))
+	UPROPERTY(EditInstanceOnly, Setter, Getter, Category="Animator", meta=(EditCondition="Mode == EPropertyAnimatorCoreMachineClockMode::Countdown", EditConditionHides))
 	FString CountdownDuration = TEXT("1m");
 
 private:
