MotionDesign : PropertyAnimator
- Removing MD_Sequence time source (replace by sequencer time source)
- Removing unused RangeTimeSource
- Renaming CurveWave animator to Curve and adjusting defaults

diff --git a/Engine/Plugins/Experimental/Avalanche/Source/AvalanchePropertyAnimator/Private/TimeSources/AvaPropertyAnimatorSequenceTimeSource.cpp b/Engine/Plugins/Experimental/Avalanche/Source/AvalanchePropertyAnimator/Private/TimeSources/AvaPropertyAnimatorSequenceTimeSource.cpp
deleted file mode 100644
diff --git a/Engine/Plugins/Experimental/Avalanche/Source/AvalanchePropertyAnimator/Private/TimeSources/AvaPropertyAnimatorSequenceTimeSource.h b/Engine/Plugins/Experimental/Avalanche/Source/AvalanchePropertyAnimator/Private/TimeSources/AvaPropertyAnimatorSequenceTimeSource.h
deleted file mode 100644
diff --git a/Engine/Plugins/Experimental/PropertyAnimator/Config/BasePropertyAnimator.ini b/Engine/Plugins/Experimental/PropertyAnimator/Config/BasePropertyAnimator.ini
--- a/Engine/Plugins/Experimental/PropertyAnimator/Config/BasePropertyAnimator.ini
+++ b/Engine/Plugins/Experimental/PropertyAnimator/Config/BasePropertyAnimator.ini
@@ -3,7 +3,7 @@
 +ClassRedirects=(OldName="/Script/PropertyController.PropertyControllerWave",NewName="/Script/PropertyAnimator.PropertyAnimatorWave")
 +ClassRedirects=(OldName="/Script/PropertyController.PropertyControllerTime",NewName="/Script/PropertyAnimator.PropertyAnimatorTime")
 +ClassRedirects=(OldName="/Script/PropertyController.PropertyControllerOscillate",NewName="/Script/PropertyAnimator.PropertyAnimatorOscillate")
-+ClassRedirects=(OldName="/Script/PropertyController.PropertyControllerFloatBase",NewName="/Script/PropertyAnimator.PropertyAnimatorFloatBase")
++ClassRedirects=(OldName="/Script/PropertyController.PropertyControllerFloatBase",NewName="/Script/PropertyAnimator.PropertyAnimatorNumericBase")
 +ClassRedirects=(OldName="/Script/PropertyController.PropertyControllerBounce",NewName="/Script/PropertyAnimator.PropertyAnimatorBounce")
 +ClassRedirects=(OldName="/Script/PropertyController.PropertyControllerBlink",NewName="/Script/PropertyAnimator.PropertyAnimatorPulse")
 +ClassRedirects=(OldName="/Script/PropertyAnimator.PropertyAnimatorBlink",NewName="/Script/PropertyAnimator.PropertyAnimatorPulse")
@@ -16,6 +16,8 @@
 +ClassRedirects=(OldName="/Script/PropertyController.PropertyControllerTimeSourceBase",NewName="/Script/PropertyAnimator.PropertyAnimatorTimeSourceBase")
 +ClassRedirects=(OldName="/Script/PropertyController.PropertyControllerWorldTimeSource",NewName="/Script/PropertyAnimator.PropertyAnimatorWorldTimeSource")
 +ClassRedirects=(OldName="/Script/AvalanchePropertyAnimator.AvaPropertyAnimatorTextResolver",NewName="/Script/PropertyAnimator.PropertyAnimatorTextResolver")
++ClassRedirects=(OldName="/Script/PropertyAnimator.PropertyAnimatorCurveWave",NewName="/Script/PropertyAnimator.PropertyAnimatorCurve")
++ClassRedirects=(OldName="/Script/PropertyAnimator.PropertyAnimatorFloatBase",NewName="/Script/PropertyAnimator.PropertyAnimatorNumericBase")
 
 +EnumRedirects=(OldName="/Script/PropertyController.EPropertyControllerEasingType",NewName="/Script/PropertyAnimator.EPropertyAnimatorEasingType")
 +EnumRedirects=(OldName="/Script/PropertyController.EPropertyControllerEasingFunction",NewName="/Script/PropertyAnimator.EPropertyAnimatorEasingFunction")
@@ -27,4 +29,4 @@
 +StructRedirects=(OldName="/Script/PropertyController.PropertyControllerWaveDoubleChannel",NewName="/Script/PropertyAnimator.PropertyAnimatorWaveDoubleChannel")
 +StructRedirects=(OldName="/Script/PropertyController.PropertyControllerWaveParameters",NewName="/Script/PropertyAnimator.PropertyAnimatorWaveParameters")
 
-+PropertyRedirects=(OldName="/Script/PropertyAnimator.PropertyAnimatorFloatBase.Strength",NewName="/Script/PropertyAnimator.PropertyAnimatorFloatBase.GlobalMagnitude")
\ No newline at end of file
++PropertyRedirects=(OldName="/Script/PropertyAnimator.PropertyAnimatorFloatBase.Strength",NewName="/Script/PropertyAnimator.PropertyAnimatorNumericBase.GlobalMagnitude")
\ No newline at end of file
diff --git a/Engine/Plugins/Experimental/PropertyAnimator/Source/PropertyAnimator/Private/Animators/PropertyAnimatorCurve.cpp b/Engine/Plugins/Experimental/PropertyAnimator/Source/PropertyAnimator/Private/Animators/PropertyAnimatorCurve.cpp
new file mode 100644
--- /dev/null
+++ b/Engine/Plugins/Experimental/PropertyAnimator/Source/PropertyAnimator/Private/Animators/PropertyAnimatorCurve.cpp
@@ -0,0 +1,136 @@
+// Copyright Epic Games, Inc. All Rights Reserved.
+
+#include "Animators/PropertyAnimatorCurve.h"
+
+#include "Curves/PropertyAnimatorEaseCurve.h"
+#include "Curves/PropertyAnimatorWaveCurve.h"
+#include "Properties/PropertyAnimatorFloatContext.h"
+#include "UObject/ConstructorHelpers.h"
+
+UPropertyAnimatorCurve::UPropertyAnimatorCurve()
+{
+	SetAnimatorDisplayName(DefaultAnimatorName);
+
+	static const ConstructorHelpers::FObjectFinder<UPropertyAnimatorWaveCurve> SineCurve(TEXT("/Script/PropertyAnimator.PropertyAnimatorWaveCurve'/PropertyAnimator/Waves/Sine.Sine'"));
+
+	if (SineCurve.Succeeded())
+	{
+		WaveCurve = SineCurve.Object;
+	}
+
+	static const ConstructorHelpers::FObjectFinder<UPropertyAnimatorEaseCurve> LinearCurve(TEXT("/Script/PropertyAnimator.PropertyAnimatorEaseCurve'/PropertyAnimator/Eases/Linear.Linear'"));
+
+	if (LinearCurve.Succeeded())
+	{
+		EaseIn.EaseCurve = LinearCurve.Object;
+		EaseOut.EaseCurve = LinearCurve.Object;
+	}
+}
+
+void UPropertyAnimatorCurve::SetWaveCurve(UPropertyAnimatorWaveCurve* InCurve)
+{
+	WaveCurve = InCurve;
+}
+
+void UPropertyAnimatorCurve::SetEaseInEnabled(bool bInEnabled)
+{
+	bEaseInEnabled = bInEnabled;
+}
+
+void UPropertyAnimatorCurve::SetEaseIn(const FPropertyAnimatorCurveEasing& InEasing)
+{
+	EaseIn = InEasing;
+}
+
+void UPropertyAnimatorCurve::SetEaseOutEnabled(bool bInEnabled)
+{
+	bEaseOutEnabled = bInEnabled;
+}
+
+void UPropertyAnimatorCurve::SetEaseOut(const FPropertyAnimatorCurveEasing& InEasing)
+{
+	EaseOut = InEasing;
+}
+
+#if WITH_EDITOR
+void UPropertyAnimatorCurve::PostEditChangeProperty(FPropertyChangedEvent& InPropertyChangedEvent)
+{
+	Super::PostEditChangeProperty(InPropertyChangedEvent);
+
+	const FName MemberName = InPropertyChangedEvent.GetMemberPropertyName();
+
+	if (MemberName == GET_MEMBER_NAME_CHECKED(UPropertyAnimatorCurve, EaseIn))
+	{
+		OnEaseInChanged();
+	}
+	else if (MemberName == GET_MEMBER_NAME_CHECKED(UPropertyAnimatorCurve, EaseOut))
+	{
+		OnEaseOutChanged();
+	}
+}
+#endif
+
+bool UPropertyAnimatorCurve::EvaluateProperty(const FPropertyAnimatorCoreData& InPropertyData, UPropertyAnimatorCoreContext* InContext, FInstancedPropertyBag& InParameters, FInstancedPropertyBag& OutEvaluationResult) const
+{
+	if (!WaveCurve)
+	{
+		return false;
+	}
+
+	const FRichCurve& SampleCurve = WaveCurve->FloatCurve;
+
+	float MinTime = 0;
+	float MaxTime = 0;
+	SampleCurve.GetTimeRange(MinTime, MaxTime);
+
+	float MinValue = 0;
+	float MaxValue = 0;
+	SampleCurve.GetValueRange(MinValue, MaxValue);
+
+	const double TimeElapsed = InParameters.GetValueDouble(TimeElapsedParameterName).GetValue();
+	const float Frequency = InParameters.GetValueDouble(FrequencyParameterName).GetValue();
+	const float Period = 1.f / Frequency;
+
+	const float SampleTime = FMath::Fmod(TimeElapsed, Period);
+	const float NormalizedSampleTime = FMath::GetMappedRangeValueClamped(FVector2D(0, Period), FVector2D(MinTime, MaxTime), SampleTime);
+
+	const float SampleValue = SampleCurve.Eval(NormalizedSampleTime);
+	float SampleValueNormalized = FMath::GetMappedRangeValueClamped(FVector2D(MinValue, MaxValue), FVector2D(0, 1), SampleValue);
+
+	if (bEaseInEnabled && EaseIn.EaseCurve && SampleTime < EaseIn.EaseDuration)
+	{
+		const FRichCurve& EaseCurve = EaseIn.EaseCurve->FloatCurve;
+		const float EaseTimeNormalized = FMath::GetMappedRangeValueClamped(FVector2D(0, EaseIn.EaseDuration), FVector2D(0, 1), SampleTime);
+		SampleValueNormalized *= EaseCurve.Eval(EaseTimeNormalized);
+	}
+
+	if (bEaseOutEnabled && EaseOut.EaseCurve && SampleTime > (CycleDuration - EaseOut.EaseDuration))
+	{
+		const FRichCurve& EaseCurve = EaseOut.EaseCurve->FloatCurve;
+		const float EaseTimeNormalized = 1 - FMath::GetMappedRangeValueClamped(FVector2D(CycleDuration - EaseOut.EaseDuration, CycleDuration), FVector2D(0, 1), SampleTime);
+		SampleValueNormalized *= EaseCurve.Eval(EaseTimeNormalized);
+	}
+
+	InParameters.AddProperty(AlphaParameterName, EPropertyBagPropertyType::Float);
+	InParameters.SetValueFloat(AlphaParameterName, SampleValueNormalized);
+
+	return InContext->EvaluateProperty(InPropertyData, InParameters, OutEvaluationResult);
+}
+
+void UPropertyAnimatorCurve::OnEaseInChanged()
+{
+	EaseIn.EaseDuration = FMath::Clamp(EaseIn.EaseDuration, 0, CycleDuration - EaseOut.EaseDuration);
+}
+
+void UPropertyAnimatorCurve::OnEaseOutChanged()
+{
+	EaseOut.EaseDuration = FMath::Clamp(EaseOut.EaseDuration, 0, CycleDuration - EaseIn.EaseDuration);
+}
+
+void UPropertyAnimatorCurve::OnCycleDurationChanged()
+{
+	Super::OnCycleDurationChanged();
+
+	OnEaseInChanged();
+	OnEaseOutChanged();
+}
diff --git a/Engine/Plugins/Experimental/PropertyAnimator/Source/PropertyAnimator/Private/Animators/PropertyAnimatorCurveWave.cpp b/Engine/Plugins/Experimental/PropertyAnimator/Source/PropertyAnimator/Private/Animators/PropertyAnimatorCurveWave.cpp
deleted file mode 100644
diff --git a/Engine/Plugins/Experimental/PropertyAnimator/Source/PropertyAnimator/Private/Animators/PropertyAnimatorFloatBase.cpp b/Engine/Plugins/Experimental/PropertyAnimator/Source/PropertyAnimator/Private/Animators/PropertyAnimatorFloatBase.cpp
deleted file mode 100644
diff --git a/Engine/Plugins/Experimental/PropertyAnimator/Source/PropertyAnimator/Private/Animators/PropertyAnimatorNumericBase.cpp b/Engine/Plugins/Experimental/PropertyAnimator/Source/PropertyAnimator/Private/Animators/PropertyAnimatorNumericBase.cpp
new file mode 100644
--- /dev/null
+++ b/Engine/Plugins/Experimental/PropertyAnimator/Source/PropertyAnimator/Private/Animators/PropertyAnimatorNumericBase.cpp
@@ -0,0 +1,256 @@
+// Copyright Epic Games, Inc. All Rights Reserved.
+
+#include "Animators/PropertyAnimatorNumericBase.h"
+
+#include "Properties/PropertyAnimatorFloatContext.h"
+#include "Properties/PropertyAnimatorRotatorContext.h"
+#include "Properties/PropertyAnimatorVectorContext.h"
+#include "Properties/Converters/PropertyAnimatorCoreConverterBase.h"
+#include "Properties/Handlers/PropertyAnimatorCoreHandlerBase.h"
+#include "Subsystems/PropertyAnimatorCoreSubsystem.h"
+
+#if WITH_EDITOR
+void UPropertyAnimatorNumericBase::PostEditChangeProperty(FPropertyChangedEvent& InPropertyChangedEvent)
+{
+	Super::PostEditChangeProperty(InPropertyChangedEvent);
+
+	const FName MemberName = InPropertyChangedEvent.GetMemberPropertyName();
+
+	static const FName RandomTimeOffsetName = GET_MEMBER_NAME_CHECKED(UPropertyAnimatorNumericBase, bRandomTimeOffset);
+	static const FName SeedName = GET_MEMBER_NAME_CHECKED(UPropertyAnimatorNumericBase, Seed);
+
+	if (MemberName == SeedName
+		|| MemberName == RandomTimeOffsetName)
+	{
+		OnSeedChanged();
+	}
+}
+#endif
+
+void UPropertyAnimatorNumericBase::SetMagnitude(float InMagnitude)
+{
+	if (FMath::IsNearlyEqual(Magnitude, InMagnitude))
+	{
+		return;
+	}
+
+	Magnitude = InMagnitude;
+	OnMagnitudeChanged();
+}
+
+void UPropertyAnimatorNumericBase::SetCycleDuration(float InCycleDuration)
+{
+	if (FMath::IsNearlyEqual(CycleDuration, InCycleDuration))
+	{
+		return;
+	}
+
+	CycleDuration = InCycleDuration;
+	OnCycleDurationChanged();
+}
+
+void UPropertyAnimatorNumericBase::SetCycleMode(EPropertyAnimatorCycleMode InMode)
+{
+	if (CycleMode == InMode)
+	{
+		return;
+	}
+
+	CycleMode = InMode;
+	OnCycleModeChanged();
+}
+
+void UPropertyAnimatorNumericBase::SetTimeOffset(double InOffset)
+{
+	if (FMath::IsNearlyEqual(TimeOffset, InOffset))
+	{
+		return;
+	}
+
+	TimeOffset = InOffset;
+	OnTimeOffsetChanged();
+}
+
+void UPropertyAnimatorNumericBase::SetRandomTimeOffset(bool bInOffset)
+{
+	if (bRandomTimeOffset == bInOffset)
+	{
+		return;
+	}
+
+	bRandomTimeOffset = bInOffset;
+	OnSeedChanged();
+}
+
+void UPropertyAnimatorNumericBase::SetSeed(int32 InSeed)
+{
+	if (Seed == InSeed)
+	{
+		return;
+	}
+
+	Seed = InSeed;
+	OnSeedChanged();
+}
+
+TSubclassOf<UPropertyAnimatorCoreContext> UPropertyAnimatorNumericBase::GetPropertyContextClass(const FPropertyAnimatorCoreData& InProperty)
+{
+	if (InProperty.IsA<FStructProperty>())
+	{
+		const FName TypeName = InProperty.GetLeafPropertyTypeName();
+
+		if (TypeName == NAME_Rotator)
+		{
+			return UPropertyAnimatorRotatorContext::StaticClass();
+		}
+
+		if (TypeName == NAME_Vector)
+		{
+			return UPropertyAnimatorVectorContext::StaticClass();
+		}
+	}
+
+	return UPropertyAnimatorFloatContext::StaticClass();
+}
+
+EPropertyAnimatorPropertySupport UPropertyAnimatorNumericBase::IsPropertySupported(const FPropertyAnimatorCoreData& InPropertyData) const
+{
+	const FName TypeName = InPropertyData.GetLeafPropertyTypeName();
+
+	if (InPropertyData.IsA<FFloatProperty>() || InPropertyData.IsA<FDoubleProperty>())
+	{
+		return EPropertyAnimatorPropertySupport::Complete;
+	}
+
+	if (InPropertyData.IsA<FStructProperty>())
+	{
+		if (TypeName == NAME_Rotator)
+		{
+			return EPropertyAnimatorPropertySupport::Complete;
+		}
+
+		if (TypeName == NAME_Vector)
+		{
+			return EPropertyAnimatorPropertySupport::Complete;
+		}
+	}
+
+	// Check if a converter supports the conversion
+	if (UPropertyAnimatorCoreSubsystem* AnimatorSubsystem = UPropertyAnimatorCoreSubsystem::Get())
+	{
+		static const FPropertyBagPropertyDesc AnimatorTypeDesc("", EPropertyBagPropertyType::Float);
+		const FPropertyBagPropertyDesc PropertyTypeDesc("", InPropertyData.GetLeafProperty());
+
+		if (AnimatorSubsystem->IsConversionSupported(AnimatorTypeDesc, PropertyTypeDesc))
+		{
+			return EPropertyAnimatorPropertySupport::Incomplete;
+		}
+	}
+
+	return Super::IsPropertySupported(InPropertyData);
+}
+
+void UPropertyAnimatorNumericBase::EvaluateProperties(FInstancedPropertyBag& InParameters)
+{
+	const float AnimatorMagnitude = Magnitude * InParameters.GetValueFloat(MagnitudeParameterName).GetValue();
+	double TimeElapsed = InParameters.GetValueDouble(TimeElapsedParameterName).GetValue();
+	RandomStream = FRandomStream(Seed);
+
+	EvaluateEachLinkedProperty<UPropertyAnimatorCoreContext>([this, &TimeElapsed, &AnimatorMagnitude, &InParameters](
+		UPropertyAnimatorCoreContext* InOptions
+		, const FPropertyAnimatorCoreData& InResolvedProperty
+		, FInstancedPropertyBag& InEvaluatedValues
+		, int32 InRangeIndex
+		, int32 InRangeMax)->bool
+	{
+		const double RandomTimeOffset = bRandomTimeOffset ? RandomStream.GetFraction() : 0;
+
+		const double AbsTimeOffset = FMath::Abs(TimeOffset);
+		const double MaxTimeOffset = InRangeMax * AbsTimeOffset;
+		double PropertyTimeElapsed = TimeElapsed - MaxTimeOffset + RandomTimeOffset;
+
+		if (TimeOffset >= 0)
+		{
+			PropertyTimeElapsed += InRangeIndex * AbsTimeOffset;
+		}
+		else
+		{
+			PropertyTimeElapsed += MaxTimeOffset - InRangeIndex * AbsTimeOffset;
+		}
+
+		if (CycleMode == EPropertyAnimatorCycleMode::DoOnce)
+		{
+			if (FMath::Abs(PropertyTimeElapsed) > CycleDuration)
+			{
+				PropertyTimeElapsed = CycleDuration - UE_KINDA_SMALL_NUMBER;
+			}
+		}
+		else if (CycleMode == EPropertyAnimatorCycleMode::Loop)
+		{
+			PropertyTimeElapsed = FMath::Fmod(PropertyTimeElapsed, CycleDuration + MaxTimeOffset + CycleGapDuration);
+
+			if (FMath::Abs(PropertyTimeElapsed) > CycleDuration)
+			{
+				PropertyTimeElapsed = CycleDuration - UE_KINDA_SMALL_NUMBER;
+			}
+		}
+		else if (CycleMode == EPropertyAnimatorCycleMode::PingPong)
+		{
+			const bool bReverse = FMath::Modulo(FMath::TruncToInt32(PropertyTimeElapsed / (CycleDuration + MaxTimeOffset + CycleGapDuration)), 2) != 0;
+			PropertyTimeElapsed = FMath::Fmod(PropertyTimeElapsed, CycleDuration + MaxTimeOffset + CycleGapDuration);
+
+			if (FMath::Abs(PropertyTimeElapsed) > CycleDuration)
+			{
+				PropertyTimeElapsed = CycleDuration - UE_KINDA_SMALL_NUMBER;
+			}
+
+			if (bReverse)
+			{
+				PropertyTimeElapsed = CycleDuration - FMath::Fmod(PropertyTimeElapsed, CycleDuration);
+			}
+			else
+			{
+				PropertyTimeElapsed = FMath::Fmod(PropertyTimeElapsed, CycleDuration);
+			}
+		}
+
+		if (Magnitude != 0
+			&& CycleDuration > 0
+			&& InOptions->GetMagnitude() != 0)
+		{
+			// Frequency
+			InParameters.AddProperty(FrequencyParameterName, EPropertyBagPropertyType::Float);
+			InParameters.SetValueFloat(FrequencyParameterName, 1.f / CycleDuration);
+
+			// Time Elapsed
+			InParameters.SetValueDouble(TimeElapsedParameterName, PropertyTimeElapsed);
+
+			// Magnitude
+			InParameters.SetValueFloat(MagnitudeParameterName, AnimatorMagnitude * InOptions->GetMagnitude());
+
+			return EvaluateProperty(InResolvedProperty, InOptions, InParameters, InEvaluatedValues);
+		}
+
+		return false;
+	});
+}
+
+void UPropertyAnimatorNumericBase::OnPropertyLinked(UPropertyAnimatorCoreContext* InLinkedProperty, EPropertyAnimatorPropertySupport InSupport)
+{
+	Super::OnPropertyLinked(InLinkedProperty, InSupport);
+
+	if (EnumHasAnyFlags(InSupport, EPropertyAnimatorPropertySupport::Incomplete))
+	{
+		if (const UPropertyAnimatorCoreSubsystem* AnimatorSubsystem = UPropertyAnimatorCoreSubsystem::Get())
+		{
+			static const FPropertyBagPropertyDesc AnimatorTypeDesc("", EPropertyBagPropertyType::Float);
+			const FPropertyBagPropertyDesc PropertyTypeDesc("", InLinkedProperty->GetAnimatedProperty().GetLeafProperty());
+
+			const TSet<UPropertyAnimatorCoreConverterBase*> Converters = AnimatorSubsystem->GetSupportedConverters(AnimatorTypeDesc, PropertyTypeDesc);
+
+			check(!Converters.IsEmpty())
+
+			InLinkedProperty->SetConverterClass(Converters.Array()[0]->GetClass());
+		}
+	}
+}
diff --git a/Engine/Plugins/Experimental/PropertyAnimator/Source/PropertyAnimator/Public/Animators/PropertyAnimatorBounce.h b/Engine/Plugins/Experimental/PropertyAnimator/Source/PropertyAnimator/Public/Animators/PropertyAnimatorBounce.h
--- a/Engine/Plugins/Experimental/PropertyAnimator/Source/PropertyAnimator/Public/Animators/PropertyAnimatorBounce.h
+++ b/Engine/Plugins/Experimental/PropertyAnimator/Source/PropertyAnimator/Public/Animators/PropertyAnimatorBounce.h
@@ -2,14 +2,14 @@
 
 #pragma once
 
-#include "Animators/PropertyAnimatorFloatBase.h"
+#include "Animators/PropertyAnimatorNumericBase.h"
 #include "PropertyAnimatorBounce.generated.h"
 
 /**
  * Applies an additive bounce movement with various options on supported float properties
  */
 UCLASS(MinimalAPI, AutoExpandCategories=("Animator"))
-class UPropertyAnimatorBounce : public UPropertyAnimatorFloatBase
+class UPropertyAnimatorBounce : public UPropertyAnimatorNumericBase
 {
 	GENERATED_BODY()
 
diff --git a/Engine/Plugins/Experimental/PropertyAnimator/Source/PropertyAnimator/Public/Animators/PropertyAnimatorCurve.h b/Engine/Plugins/Experimental/PropertyAnimator/Source/PropertyAnimator/Public/Animators/PropertyAnimatorCurve.h
new file mode 100644
--- /dev/null
+++ b/Engine/Plugins/Experimental/PropertyAnimator/Source/PropertyAnimator/Public/Animators/PropertyAnimatorCurve.h
@@ -0,0 +1,103 @@
+// Copyright Epic Games, Inc. All Rights Reserved.
+
+#pragma once
+
+#include "Animators/PropertyAnimatorNumericBase.h"
+#include "PropertyAnimatorCurve.generated.h"
+
+class UPropertyAnimatorEaseCurve;
+class UPropertyAnimatorWaveCurve;
+
+USTRUCT(BlueprintType)
+struct FPropertyAnimatorCurveEasing
+{
+	GENERATED_BODY()
+
+	UPROPERTY(EditInstanceOnly, Category="Animator")
+	TObjectPtr<UPropertyAnimatorEaseCurve> EaseCurve;
+
+	UPROPERTY(EditInstanceOnly, Interp, Category="Animator", meta=(ClampMin="0", Units=Seconds))
+	float EaseDuration = 0.f;
+};
+
+/**
+ * Applies a wave movement from a curve on supported float properties
+ */
+UCLASS(MinimalAPI, AutoExpandCategories=("Animator"))
+class UPropertyAnimatorCurve : public UPropertyAnimatorNumericBase
+{
+	GENERATED_BODY()
+
+public:
+	static constexpr const TCHAR* DefaultAnimatorName = TEXT("Curve");
+
+	UPropertyAnimatorCurve();
+
+	void SetWaveCurve(UPropertyAnimatorWaveCurve* InCurve);
+
+	UPropertyAnimatorWaveCurve* GetWaveCurve() const
+	{
+		return WaveCurve;
+	}
+
+	void SetEaseInEnabled(bool bInEnabled);
+
+	bool GetEaseInEnabled() const
+	{
+		return bEaseInEnabled;
+	}
+
+	void SetEaseIn(const FPropertyAnimatorCurveEasing& InEasing);
+
+	const FPropertyAnimatorCurveEasing& GetEaseIn() const
+	{
+		return EaseIn;
+	}
+
+	void SetEaseOutEnabled(bool bInEnabled);
+
+	bool GetEaseOutEnabled() const
+	{
+		return bEaseOutEnabled;
+	}
+
+	void SetEaseOut(const FPropertyAnimatorCurveEasing& InEasing);
+
+	const FPropertyAnimatorCurveEasing& GetEaseOut() const
+	{
+		return EaseOut;
+	}
+
+protected:
+	//~ Begin UObject
+#if WITH_EDITOR
+	virtual void PostEditChangeProperty(FPropertyChangedEvent& InPropertyChangedEvent) override;
+#endif
+	//~ End UObject
+
+	//~ Begin UPropertyAnimatorFloatBase
+	virtual bool EvaluateProperty(const FPropertyAnimatorCoreData& InPropertyData, UPropertyAnimatorCoreContext* InContext, FInstancedPropertyBag& InParameters, FInstancedPropertyBag& OutEvaluationResult) const override;
+	//~ End UPropertyAnimatorFloatBase
+
+	void OnEaseInChanged();
+	void OnEaseOutChanged();
+	virtual void OnCycleDurationChanged() override;
+
+	/** The wave curve to sample for the animation */
+	UPROPERTY(EditInstanceOnly, Setter, Getter, Category="Animator")
+	TObjectPtr<UPropertyAnimatorWaveCurve> WaveCurve;
+
+	UPROPERTY(EditInstanceOnly, Setter="SetEaseInEnabled", Getter="GetEaseInEnabled", Category="Animator", meta=(InlineEditConditionToggle))
+	bool bEaseInEnabled = false;
+
+	/** Ease in for this effect */
+	UPROPERTY(EditInstanceOnly, Setter, Getter, Category="Animator", meta=(EditCondition="bEaseInEnabled"))
+	FPropertyAnimatorCurveEasing EaseIn;
+
+	UPROPERTY(EditInstanceOnly, Setter="SetEaseOutEnabled", Getter="GetEaseOutEnabled", Category="Animator", meta=(InlineEditConditionToggle))
+	bool bEaseOutEnabled = false;
+
+	/** Ease out for this effect */
+	UPROPERTY(EditInstanceOnly, Setter, Getter, Category="Animator", meta=(EditCondition="bEaseOutEnabled"))
+    FPropertyAnimatorCurveEasing EaseOut;
+};
\ No newline at end of file
diff --git a/Engine/Plugins/Experimental/PropertyAnimator/Source/PropertyAnimator/Public/Animators/PropertyAnimatorCurveWave.h b/Engine/Plugins/Experimental/PropertyAnimator/Source/PropertyAnimator/Public/Animators/PropertyAnimatorCurveWave.h
deleted file mode 100644
diff --git a/Engine/Plugins/Experimental/PropertyAnimator/Source/PropertyAnimator/Public/Animators/PropertyAnimatorFloatBase.h b/Engine/Plugins/Experimental/PropertyAnimator/Source/PropertyAnimator/Public/Animators/PropertyAnimatorFloatBase.h
deleted file mode 100644
diff --git a/Engine/Plugins/Experimental/PropertyAnimator/Source/PropertyAnimator/Public/Animators/PropertyAnimatorNumericBase.h b/Engine/Plugins/Experimental/PropertyAnimator/Source/PropertyAnimator/Public/Animators/PropertyAnimatorNumericBase.h
new file mode 100644
--- /dev/null
+++ b/Engine/Plugins/Experimental/PropertyAnimator/Source/PropertyAnimator/Public/Animators/PropertyAnimatorNumericBase.h
@@ -0,0 +1,123 @@
+// Copyright Epic Games, Inc. All Rights Reserved.
+
+#pragma once
+
+#include "Animators/PropertyAnimatorCoreBase.h"
+#include "PropertyAnimatorNumericBase.generated.h"
+
+UENUM(BlueprintType)
+enum class EPropertyAnimatorCycleMode : uint8
+{
+	/** Cycle only once then stop */
+	DoOnce,
+	/** Cycle and repeat once we reached the end */
+	Loop,
+	/** Cycle and reverse repeat */
+	PingPong
+};
+
+/**
+ * Animate supported numeric properties with various options
+ */
+UCLASS(MinimalAPI, Abstract, AutoExpandCategories=("Animator"))
+class UPropertyAnimatorNumericBase : public UPropertyAnimatorCoreBase
+{
+	GENERATED_BODY()
+
+	friend class FPropertyAnimatorCoreEditorDetailCustomization;
+
+public:
+	PROPERTYANIMATOR_API void SetMagnitude(float InMagnitude);
+	float GetMagnitude() const
+	{
+		return Magnitude;
+	}
+
+	PROPERTYANIMATOR_API void SetCycleDuration(float InCycleDuration);
+	float GetCycleDuration() const
+	{
+		return CycleDuration;
+	}
+
+	PROPERTYANIMATOR_API void SetCycleMode(EPropertyAnimatorCycleMode InMode);
+	EPropertyAnimatorCycleMode GetCycleMode() const
+	{
+		return CycleMode;
+	}
+
+	PROPERTYANIMATOR_API void SetTimeOffset(double InOffset);
+	double GetTimeOffset() const
+	{
+		return TimeOffset;
+	}
+
+	PROPERTYANIMATOR_API void SetRandomTimeOffset(bool bInOffset);
+	bool GetRandomTimeOffset() const
+	{
+		return bRandomTimeOffset;
+	}
+
+	PROPERTYANIMATOR_API void SetSeed(int32 InSeed);
+	int32 GetSeed() const
+	{
+		return Seed;
+	}
+
+protected:
+	//~ Begin UObject
+#if WITH_EDITOR
+	virtual void PostEditChangeProperty(FPropertyChangedEvent& InPropertyChangedEvent) override;
+#endif
+	//~ End UObject
+
+	virtual void OnMagnitudeChanged() {}
+	virtual void OnCycleDurationChanged() {}
+	virtual void OnCycleModeChanged() {}
+	virtual void OnTimeOffsetChanged() {}
+	virtual void OnSeedChanged() {}
+
+	//~ Begin UPropertyAnimatorCoreBase
+	virtual TSubclassOf<UPropertyAnimatorCoreContext> GetPropertyContextClass(const FPropertyAnimatorCoreData& InProperty) override;
+	virtual EPropertyAnimatorPropertySupport IsPropertySupported(const FPropertyAnimatorCoreData& InPropertyData) const override;
+	virtual void EvaluateProperties(FInstancedPropertyBag& InParameters) override;
+	virtual void OnPropertyLinked(UPropertyAnimatorCoreContext* InLinkedProperty, EPropertyAnimatorPropertySupport InSupport) override;
+	//~ End UPropertyAnimatorCoreBase
+
+	/** Evaluate and return float value for a property */
+	virtual bool EvaluateProperty(const FPropertyAnimatorCoreData& InPropertyData, UPropertyAnimatorCoreContext* InContext, FInstancedPropertyBag& InParameters, FInstancedPropertyBag& OutEvaluationResult) const
+	{
+		return false;
+	}
+
+	/** Magnitude for the effect on all properties */
+	UPROPERTY(EditInstanceOnly, Setter, Getter, Category="Animator", meta=(ClampMin="0"))
+	float Magnitude = 1.f;
+
+	/** Duration of one cycle for the effect = period of the effect */
+	UPROPERTY(EditInstanceOnly, Setter, Getter, Category="Animator", meta=(ClampMin="0", Units=Seconds))
+	float CycleDuration = 1.f;
+
+	/** Cycle mode for the effect */
+	UPROPERTY(EditInstanceOnly, Setter, Getter, Category="Animator")
+	EPropertyAnimatorCycleMode CycleMode = EPropertyAnimatorCycleMode::Loop;
+
+	/** Time gap between each cycle */
+	UPROPERTY(EditInstanceOnly, Category="Animator", meta=(ClampMin="0", Units=Seconds, EditCondition="CycleMode != EPropertyAnimatorCycleMode::DoOnce", EditConditionHides))
+	float CycleGapDuration = 0.f;
+
+	/** Use random time offset to add variation in animation */
+	UPROPERTY(EditInstanceOnly, Setter="SetRandomTimeOffset", Getter="GetRandomTimeOffset", Category="Animator", meta=(InlineEditConditionToggle))
+	bool bRandomTimeOffset = false;
+
+	/** Seed to generate per property time offset */
+	UPROPERTY(EditInstanceOnly, Setter, Getter, Category="Animator", meta=(EditCondition="bRandomTimeOffset"))
+	int32 Seed = 0;
+
+	/** Time offset accumulated for each property for every round */
+	UPROPERTY(EditInstanceOnly, Setter, Getter, Category="Animator", meta=(Units=Seconds))
+	double TimeOffset = 0;
+
+private:
+	/** Random stream for time offset */
+	FRandomStream RandomStream = FRandomStream(Seed);
+};
\ No newline at end of file
diff --git a/Engine/Plugins/Experimental/PropertyAnimator/Source/PropertyAnimator/Public/Animators/PropertyAnimatorOscillate.h b/Engine/Plugins/Experimental/PropertyAnimator/Source/PropertyAnimator/Public/Animators/PropertyAnimatorOscillate.h
--- a/Engine/Plugins/Experimental/PropertyAnimator/Source/PropertyAnimator/Public/Animators/PropertyAnimatorOscillate.h
+++ b/Engine/Plugins/Experimental/PropertyAnimator/Source/PropertyAnimator/Public/Animators/PropertyAnimatorOscillate.h
@@ -2,7 +2,7 @@
 
 #pragma once
 
-#include "Animators/PropertyAnimatorFloatBase.h"
+#include "Animators/PropertyAnimatorNumericBase.h"
 #include "PropertyAnimatorOscillate.generated.h"
 
 UENUM(BlueprintType)
@@ -20,7 +20,7 @@ enum class EPropertyAnimatorOscillateFunction : uint8
  * Applies an additive regular oscillate movement with various options on supported float properties
  */
 UCLASS(MinimalAPI, AutoExpandCategories=("Animator"))
-class UPropertyAnimatorOscillate : public UPropertyAnimatorFloatBase
+class UPropertyAnimatorOscillate : public UPropertyAnimatorNumericBase
 {
 	GENERATED_BODY()
 
diff --git a/Engine/Plugins/Experimental/PropertyAnimator/Source/PropertyAnimator/Public/Animators/PropertyAnimatorPulse.h b/Engine/Plugins/Experimental/PropertyAnimator/Source/PropertyAnimator/Public/Animators/PropertyAnimatorPulse.h
--- a/Engine/Plugins/Experimental/PropertyAnimator/Source/PropertyAnimator/Public/Animators/PropertyAnimatorPulse.h
+++ b/Engine/Plugins/Experimental/PropertyAnimator/Source/PropertyAnimator/Public/Animators/PropertyAnimatorPulse.h
@@ -4,7 +4,7 @@
 
 #pragma once
 
-#include "Animators/PropertyAnimatorFloatBase.h"
+#include "Animators/PropertyAnimatorNumericBase.h"
 #include "PropertyAnimatorShared.h"
 #include "PropertyAnimatorPulse.generated.h"
 
@@ -12,7 +12,7 @@
  * Applies an additive pulse movement with various options on supported float properties
  */
 UCLASS(MinimalAPI, AutoExpandCategories=("Animator"))
-class UPropertyAnimatorPulse : public UPropertyAnimatorFloatBase
+class UPropertyAnimatorPulse : public UPropertyAnimatorNumericBase
 {
 	GENERATED_BODY()
 
diff --git a/Engine/Plugins/Experimental/PropertyAnimator/Source/PropertyAnimator/Public/Animators/PropertyAnimatorSoundWave.h b/Engine/Plugins/Experimental/PropertyAnimator/Source/PropertyAnimator/Public/Animators/PropertyAnimatorSoundWave.h
--- a/Engine/Plugins/Experimental/PropertyAnimator/Source/PropertyAnimator/Public/Animators/PropertyAnimatorSoundWave.h
+++ b/Engine/Plugins/Experimental/PropertyAnimator/Source/PropertyAnimator/Public/Animators/PropertyAnimatorSoundWave.h
@@ -2,7 +2,7 @@
 
 #pragma once
 
-#include "Animators/PropertyAnimatorFloatBase.h"
+#include "Animators/PropertyAnimatorNumericBase.h"
 #include "PropertyAnimatorSoundWave.generated.h"
 
 class ULoudnessNRT;
@@ -12,7 +12,7 @@ class USoundWave;
  * Applies a sampled sound wave movement with various options on supported float properties
  */
 UCLASS(MinimalAPI, AutoExpandCategories=("Animator"))
-class UPropertyAnimatorSoundWave : public UPropertyAnimatorFloatBase
+class UPropertyAnimatorSoundWave : public UPropertyAnimatorNumericBase
 {
 	GENERATED_BODY()
 
diff --git a/Engine/Plugins/Experimental/PropertyAnimator/Source/PropertyAnimator/Public/Animators/PropertyAnimatorTime.h b/Engine/Plugins/Experimental/PropertyAnimator/Source/PropertyAnimator/Public/Animators/PropertyAnimatorTime.h
--- a/Engine/Plugins/Experimental/PropertyAnimator/Source/PropertyAnimator/Public/Animators/PropertyAnimatorTime.h
+++ b/Engine/Plugins/Experimental/PropertyAnimator/Source/PropertyAnimator/Public/Animators/PropertyAnimatorTime.h
@@ -2,14 +2,14 @@
 
 #pragma once
 
-#include "Animators/PropertyAnimatorFloatBase.h"
+#include "Animators/PropertyAnimatorNumericBase.h"
 #include "PropertyAnimatorTime.generated.h"
 
 /**
  * Applies an additive time movement with various options on supported float properties
  */
 UCLASS(MinimalAPI, AutoExpandCategories=("Animator"))
-class UPropertyAnimatorTime : public UPropertyAnimatorFloatBase
+class UPropertyAnimatorTime : public UPropertyAnimatorNumericBase
 {
 	GENERATED_BODY()
 
diff --git a/Engine/Plugins/Experimental/PropertyAnimator/Source/PropertyAnimator/Public/Animators/PropertyAnimatorWiggle.h b/Engine/Plugins/Experimental/PropertyAnimator/Source/PropertyAnimator/Public/Animators/PropertyAnimatorWiggle.h
--- a/Engine/Plugins/Experimental/PropertyAnimator/Source/PropertyAnimator/Public/Animators/PropertyAnimatorWiggle.h
+++ b/Engine/Plugins/Experimental/PropertyAnimator/Source/PropertyAnimator/Public/Animators/PropertyAnimatorWiggle.h
@@ -2,14 +2,14 @@
 
 #pragma once
 
-#include "Animators/PropertyAnimatorFloatBase.h"
+#include "Animators/PropertyAnimatorNumericBase.h"
 #include "PropertyAnimatorWiggle.generated.h"
 
 /**
  * Applies an additive random wiggle movement with various options on supported float properties
  */
 UCLASS(MinimalAPI, AutoExpandCategories=("Animator"))
-class UPropertyAnimatorWiggle : public UPropertyAnimatorFloatBase
+class UPropertyAnimatorWiggle : public UPropertyAnimatorNumericBase
 {
 	GENERATED_BODY()
 
diff --git a/Engine/Plugins/Experimental/PropertyAnimatorCore/Config/BasePropertyAnimatorCore.ini b/Engine/Plugins/Experimental/PropertyAnimatorCore/Config/BasePropertyAnimatorCore.ini
--- a/Engine/Plugins/Experimental/PropertyAnimatorCore/Config/BasePropertyAnimatorCore.ini
+++ b/Engine/Plugins/Experimental/PropertyAnimatorCore/Config/BasePropertyAnimatorCore.ini
@@ -11,6 +11,7 @@
 +ClassRedirects=(OldName="/Script/PropertyAnimator.PropertyAnimatorCustomTimeSource",NewName="/Script/PropertyAnimatorCore.PropertyAnimatorCoreManualTimeSource")
 +ClassRedirects=(OldName="/Script/PropertyAnimator.PropertyAnimatorTimeSourceBase",NewName="/Script/PropertyAnimatorCore.PropertyAnimatorCoreTimeSourceBase")
 +ClassRedirects=(OldName="/Script/PropertyAnimator.PropertyAnimatorWorldTimeSource",NewName="/Script/PropertyAnimatorCore.PropertyAnimatorCoreWorldTimeSource")
++ClassRedirects=(OldName="/Script/AvalanchePropertyAnimator.AvaPropertyAnimatorSequenceTimeSource",NewName="/Script/PropertyAnimatorCore.PropertyAnimatorCoreSequencerTimeSource")
 
 +StructRedirects=(OldName="/Script/CustomPropertyControlEditor.CustomPropertyControlData",NewName="/Script/PropertyAnimatorCore.PropertyAnimatorCoreData")
 
diff --git a/Engine/Plugins/Experimental/PropertyAnimatorCore/Source/PropertyAnimatorCore/Private/TimeSources/PropertyAnimatorCoreRangeTimeSource.cpp b/Engine/Plugins/Experimental/PropertyAnimatorCore/Source/PropertyAnimatorCore/Private/TimeSources/PropertyAnimatorCoreRangeTimeSource.cpp
deleted file mode 100644
diff --git a/Engine/Plugins/Experimental/PropertyAnimatorCore/Source/PropertyAnimatorCore/Public/TimeSources/PropertyAnimatorCoreRangeTimeSource.h b/Engine/Plugins/Experimental/PropertyAnimatorCore/Source/PropertyAnimatorCore/Public/TimeSources/PropertyAnimatorCoreRangeTimeSource.h
deleted file mode 100644
