MotionDesign : PropertyAnimator
- Adding new datetime formats to clock animator with and without leading zeros
- (Un)register new datetime formats easily

diff --git a/Engine/Plugins/Experimental/PropertyAnimator/Source/PropertyAnimator/Private/Animators/PropertyAnimatorClock.cpp b/Engine/Plugins/Experimental/PropertyAnimator/Source/PropertyAnimator/Private/Animators/PropertyAnimatorClock.cpp
--- a/Engine/Plugins/Experimental/PropertyAnimator/Source/PropertyAnimator/Private/Animators/PropertyAnimatorClock.cpp
+++ b/Engine/Plugins/Experimental/PropertyAnimator/Source/PropertyAnimator/Private/Animators/PropertyAnimatorClock.cpp
@@ -7,9 +7,54 @@
 #include "Properties/Handlers/PropertyAnimatorCoreHandlerBase.h"
 #include "Subsystems/PropertyAnimatorCoreSubsystem.h"
 
+namespace UE::Private::Animator
+{
+	static TMap<TCHAR, TFunction<FString(const FDateTime&)>> FormatFunctions;
+}
+
+void UPropertyAnimatorClock::RegisterFormat(const TCHAR InChar, TFunction<FString(const FDateTime&)> InFormatter)
+{
+	UE::Private::Animator::FormatFunctions.Add(InChar, InFormatter);
+}
+
+void UPropertyAnimatorClock::UnregisterFormat(const TCHAR InChar)
+{
+	UE::Private::Animator::FormatFunctions.Remove(InChar);
+}
+
 UPropertyAnimatorClock::UPropertyAnimatorClock()
 {
 	SetAnimatorDisplayName(DefaultControllerName);
+
+	if (IsTemplate())
+	{
+		RegisterFormat(TEXT('a'), [](const FDateTime& InDateTime)->FString{ return InDateTime.ToFormattedString(TEXT("%a")); });
+		RegisterFormat(TEXT('A'), [](const FDateTime& InDateTime)->FString{ return InDateTime.ToFormattedString(TEXT("%A")); });
+		RegisterFormat(TEXT('w'), [](const FDateTime& InDateTime)->FString{ return InDateTime.ToFormattedString(TEXT("%w")); });
+		RegisterFormat(TEXT('y'), [](const FDateTime& InDateTime)->FString{ return InDateTime.ToFormattedString(TEXT("%y")); });
+		RegisterFormat(TEXT('Y'), [](const FDateTime& InDateTime)->FString{ return InDateTime.ToFormattedString(TEXT("%Y")); });
+		RegisterFormat(TEXT('b'), [](const FDateTime& InDateTime)->FString{ return InDateTime.ToFormattedString(TEXT("%b")); });
+		RegisterFormat(TEXT('B'), [](const FDateTime& InDateTime)->FString{ return InDateTime.ToFormattedString(TEXT("%B")); });
+		RegisterFormat(TEXT('m'), [](const FDateTime& InDateTime)->FString{ return InDateTime.ToFormattedString(TEXT("%m")); });
+		RegisterFormat(TEXT('n'), [](const FDateTime& InDateTime)->FString{ return FString::FromInt(InDateTime.GetMonth()); });
+		RegisterFormat(TEXT('d'), [](const FDateTime& InDateTime)->FString{ return InDateTime.ToFormattedString(TEXT("%d")); });
+		RegisterFormat(TEXT('e'), [](const FDateTime& InDateTime)->FString{ return InDateTime.ToFormattedString(TEXT("%e")); });
+		RegisterFormat(TEXT('j'), [](const FDateTime& InDateTime)->FString{ return InDateTime.ToFormattedString(TEXT("%j")); });
+		RegisterFormat(TEXT('J'), [](const FDateTime& InDateTime)->FString{ return FString::FromInt(InDateTime.GetDayOfYear()); });
+		RegisterFormat(TEXT('l'), [](const FDateTime& InDateTime)->FString{ return InDateTime.ToFormattedString(TEXT("%l")); });
+		RegisterFormat(TEXT('I'), [](const FDateTime& InDateTime)->FString{ return InDateTime.ToFormattedString(TEXT("%I")); });
+		RegisterFormat(TEXT('H'), [](const FDateTime& InDateTime)->FString{ return InDateTime.ToFormattedString(TEXT("%H")); });
+		RegisterFormat(TEXT('h'), [](const FDateTime& InDateTime)->FString{ return FString::FromInt(InDateTime.GetHour()); });
+		RegisterFormat(TEXT('M'), [](const FDateTime& InDateTime)->FString{ return InDateTime.ToFormattedString(TEXT("%M")); });
+		RegisterFormat(TEXT('N'), [](const FDateTime& InDateTime)->FString{ return FString::FromInt(InDateTime.GetMinute()); });
+		RegisterFormat(TEXT('S'), [](const FDateTime& InDateTime)->FString{ return InDateTime.ToFormattedString(TEXT("%S")); });
+		RegisterFormat(TEXT('s'), [](const FDateTime& InDateTime)->FString{ return FString::FromInt(InDateTime.GetSecond()); });
+		RegisterFormat(TEXT('f'), [](const FDateTime& InDateTime)->FString{ return InDateTime.ToString(TEXT("%s")); });
+		RegisterFormat(TEXT('F'), [](const FDateTime& InDateTime)->FString{ return FString::FromInt(InDateTime.GetMillisecond()); });
+		RegisterFormat(TEXT('p'), [](const FDateTime& InDateTime)->FString{ return InDateTime.ToFormattedString(TEXT("%p")); });
+		RegisterFormat(TEXT('P'), [](const FDateTime& InDateTime)->FString{ return InDateTime.ToFormattedString(TEXT("%P")); });
+		RegisterFormat(TEXT('t'), [](const FDateTime& InDateTime)->FString{ return LexToString(InDateTime.GetTicks()); });
+	}
 }
 
 void UPropertyAnimatorClock::SetDisplayFormat(const FString& InDisplayFormat)
@@ -17,6 +62,31 @@ void UPropertyAnimatorClock::SetDisplayFormat(const FString& InDisplayFormat)
 	DisplayFormat = InDisplayFormat;
 }
 
+FString UPropertyAnimatorClock::FormatDateTime(const FDateTime& InDateTime, const FString& InDisplayFormat)
+{
+	FString Result;
+	Result.Reserve(InDisplayFormat.Len());
+
+    for (int32 CharIndex = 0; CharIndex < InDisplayFormat.Len(); ++CharIndex)
+    {
+        if (InDisplayFormat[CharIndex] == TEXT('%') && CharIndex + 1 < InDisplayFormat.Len())
+        {
+			if (const TFunction<FString(const FDateTime&)>* Formatter = UE::Private::Animator::FormatFunctions.Find(InDisplayFormat[CharIndex + 1]))
+			{
+				Result += (*Formatter)(InDateTime);
+			}
+
+            ++CharIndex;
+        }
+        else
+        {
+            Result += InDisplayFormat[CharIndex];
+        }
+    }
+
+    return Result;
+}
+
 EPropertyAnimatorPropertySupport UPropertyAnimatorClock::IsPropertySupported(const FPropertyAnimatorCoreData& InPropertyData) const
 {
 	if (InPropertyData.IsA<FStrProperty>())
@@ -45,7 +115,7 @@ void UPropertyAnimatorClock::EvaluateProperties(FInstancedPropertyBag& InParamet
 
 	const FTimespan ElapsedTimeSpan = FTimespan::FromSeconds(TimeElapsed);
 	const FDateTime DateTime(ElapsedTimeSpan > FTimespan::Zero() ? ElapsedTimeSpan.GetTicks() : 0);
-	const FString FormattedDateTime = DateTime.ToFormattedString(*DisplayFormat);
+	const FString FormattedDateTime = FormatDateTime(DateTime, DisplayFormat);
 
 	EvaluateEachLinkedProperty<UPropertyAnimatorCoreContext>([this, FormattedDateTime](
 		UPropertyAnimatorCoreContext* InContext
@@ -67,20 +137,17 @@ void UPropertyAnimatorClock::OnPropertyLinked(UPropertyAnimatorCoreContext* InLi
 {
 	Super::OnPropertyLinked(InLinkedProperty, InSupport);
 
-	const FPropertyAnimatorCoreData& Property = InLinkedProperty->GetAnimatedProperty();
-	if (Property.IsA<FStrProperty>())
-	{
-		return;
-	}
-
 	if (EnumHasAnyFlags(InSupport, EPropertyAnimatorPropertySupport::Incomplete))
 	{
 		if (const UPropertyAnimatorCoreSubsystem* AnimatorSubsystem = UPropertyAnimatorCoreSubsystem::Get())
 		{
 			static const FPropertyBagPropertyDesc AnimatorTypeDesc("", EPropertyBagPropertyType::String);
+
+			const FPropertyAnimatorCoreData& Property = InLinkedProperty->GetAnimatedProperty();
 			const FPropertyBagPropertyDesc PropertyTypeDesc("", Property.GetLeafProperty());
 			const TSet<UPropertyAnimatorCoreConverterBase*> Converters = AnimatorSubsystem->GetSupportedConverters(AnimatorTypeDesc, PropertyTypeDesc);
 			check(!Converters.IsEmpty())
+
 			InLinkedProperty->SetConverterClass(Converters.Array()[0]->GetClass());
 		}
 	}
diff --git a/Engine/Plugins/Experimental/PropertyAnimator/Source/PropertyAnimator/Public/Animators/PropertyAnimatorClock.h b/Engine/Plugins/Experimental/PropertyAnimator/Source/PropertyAnimator/Public/Animators/PropertyAnimatorClock.h
--- a/Engine/Plugins/Experimental/PropertyAnimator/Source/PropertyAnimator/Public/Animators/PropertyAnimatorClock.h
+++ b/Engine/Plugins/Experimental/PropertyAnimator/Source/PropertyAnimator/Public/Animators/PropertyAnimatorClock.h
@@ -3,6 +3,7 @@
 #pragma once
 
 #include "Animators/PropertyAnimatorCoreBase.h"
+#include "HAL/Platform.h"
 #include "Misc/Timespan.h"
 #include "PropertyAnimatorClock.generated.h"
 
@@ -27,6 +28,10 @@ class UPropertyAnimatorClock : public UPropertyAnimatorCoreBase
 public:
 	static constexpr const TCHAR* DefaultControllerName = TEXT("Clock");
 
+	static void RegisterFormat(TCHAR InChar, TFunction<FString(const FDateTime&)> InFormatter);
+	static void UnregisterFormat(TCHAR InChar);
+	static FString FormatDateTime(const FDateTime& InDateTime, const FString& InDisplayFormat);
+
 	UPropertyAnimatorClock();
 
 	PROPERTYANIMATOR_API void SetDisplayFormat(const FString& InDisplayFormat);
@@ -52,16 +57,24 @@ protected:
 	 * %b - Month, eg) Jan
 	 * %B - Month, eg) January
 	 * %m - Month, 01-12
+	 * %n - Month, 1-12
 	 * %d - Day, 01-31
 	 * %e - Day, 1-31
+	 * %j - Day of the Year, 001-366
+	 * %J - Day of the Year, 1-366
 	 * %l - 12h Hour, 1-12
 	 * %I - 12h Hour, 01-12
 	 * %H - 24h Hour, 00-23
+	 * %h - 24h Hour, 0-23
 	 * %M - Minute, 00-59
+	 * %N - Minute, 0-59
 	 * %S - Second, 00-60
+	 * %s - Second, 0-60
+	 * %f - Millisecond, 000-999
+	 * %F - Millisecond, 0-999
 	 * %p - AM or PM
 	 * %P - am or PM
-	 * %j - Day of the Year, 001-366
+	 * %t - Ticks since midnight, January 1, 0001
 	 */
 	UPROPERTY(EditInstanceOnly, BlueprintReadWrite, Setter, Getter, Category="Animator")
 	FString DisplayFormat = TEXT("%H:%M:%S");
