MotionDesign : PropertyAnimator
- Adding counter animator with various options to display a formatted number
- Formats are stored in the developer settings and can be reused in other animators

diff --git a/Engine/Plugins/Experimental/PropertyAnimator/Source/PropertyAnimator/Private/Animators/PropertyAnimatorCounter.cpp b/Engine/Plugins/Experimental/PropertyAnimator/Source/PropertyAnimator/Private/Animators/PropertyAnimatorCounter.cpp
new file mode 100644
--- /dev/null
+++ b/Engine/Plugins/Experimental/PropertyAnimator/Source/PropertyAnimator/Private/Animators/PropertyAnimatorCounter.cpp
@@ -0,0 +1,344 @@
+// Copyright Epic Games, Inc. All Rights Reserved.
+
+#include "Animators/PropertyAnimatorCounter.h"
+
+#include "Properties/Converters/PropertyAnimatorCoreConverterBase.h"
+#include "Settings/PropertyAnimatorSettings.h"
+#include "Subsystems/PropertyAnimatorCoreSubsystem.h"
+
+void FPropertyAnimatorCounterFormat::EnsureCharactersLength()
+{
+	if (DecimalCharacter.Len() > 1)
+	{
+		DecimalCharacter.RemoveAt(1, DecimalCharacter.Len() - 1);
+	}
+
+	if (PaddingCharacter.Len() > 1)
+	{
+		PaddingCharacter.RemoveAt(1, PaddingCharacter.Len() - 1);
+	}
+
+	if (GroupingCharacter.Len() > 1)
+	{
+		GroupingCharacter.RemoveAt(1, GroupingCharacter.Len() - 1);
+	}
+}
+
+FString FPropertyAnimatorCounterFormat::FormatNumber(double InNumber) const
+{
+	const bool bPositive = InNumber >= 0;
+	InNumber = FMath::Abs(InNumber);
+
+	if (RoundingMode != EPropertyAnimatorCounterRoundingMode::None)
+	{
+		switch (RoundingMode)
+		{
+		case EPropertyAnimatorCounterRoundingMode::Round:
+			InNumber = FMath::RoundToDouble(InNumber);
+		break;
+		case EPropertyAnimatorCounterRoundingMode::Floor:
+			InNumber = FMath::FloorToDouble(InNumber);
+		break;
+		case EPropertyAnimatorCounterRoundingMode::Ceil:
+			InNumber = FMath::CeilToDouble(InNumber);
+		break;
+		default:
+		break;
+		}
+	}
+
+	const int32 IntegerCount = MinIntegerCount;
+	const int32 DecimalCount = MaxDecimalCount;
+
+	FString IntegerPart = LexToString(FMath::TruncToInt(InNumber));
+	FString DecimalPart;
+
+	if (IntegerPart.Len() > IntegerCount && bTruncate)
+	{
+		IntegerPart.RemoveAt(0, IntegerPart.Len() - IntegerCount);
+	}
+
+	if (DecimalCount > 0)
+	{
+		DecimalPart = FString::SanitizeFloat(FMath::Fractional(InNumber));
+
+		DecimalPart.RemoveFromStart(TEXT("0."));
+
+		if (DecimalPart.Len() < DecimalCount)
+		{
+			DecimalPart.InsertAt(DecimalPart.Len() - 1, FString::ChrN(DecimalCount - DecimalPart.Len(), TEXT('0')));
+		}
+
+		if (DecimalPart.Len() > DecimalCount)
+		{
+			DecimalPart.RemoveAt(DecimalCount, DecimalPart.Len());
+		}
+
+		if (!DecimalCharacter.IsEmpty())
+		{
+			DecimalPart.InsertAt(0, DecimalCharacter[0]);
+		}
+	}
+
+	if (!PaddingCharacter.IsEmpty() && IntegerPart.Len() < IntegerCount)
+	{
+		IntegerPart.InsertAt(0, FString::ChrN(IntegerCount - IntegerPart.Len(), PaddingCharacter[0]));
+	}
+
+	if (!GroupingCharacter.IsEmpty())
+	{
+		int32 ThousandsIndex = IntegerPart.Len() - 1;
+		int32 ThousandsCount = 0;
+		while (ThousandsIndex > 0)
+		{
+			if (IntegerPart[ThousandsIndex] != TEXT('.'))
+			{
+				ThousandsCount++;
+
+				if (FMath::Modulo(ThousandsCount, GroupingSize) == 0 && ThousandsIndex != 0)
+				{
+					IntegerPart.InsertAt(ThousandsIndex, GroupingCharacter[0]);
+				}
+			}
+
+			ThousandsIndex--;
+		}
+	}
+
+	FString NumberString = IntegerPart + DecimalPart;
+
+	if (bUseSign)
+	{
+		NumberString.InsertAt(0, bPositive ? TEXT('+') : TEXT('-'));
+	}
+
+	return NumberString;
+}
+
+UPropertyAnimatorCounter::UPropertyAnimatorCounter()
+{
+	SetAnimatorDisplayName(DefaultAnimatorName);
+
+	if (!IsTemplate())
+	{
+		const TArray<FName> AvailableNames = GetAvailableFormatNames();
+
+		if (!AvailableNames.IsEmpty())
+		{
+			PresetFormatName = AvailableNames[0];
+		}
+	}
+}
+
+void UPropertyAnimatorCounter::SetDisplayPattern(const FText& InPattern)
+{
+	DisplayPattern = InPattern;
+}
+
+void UPropertyAnimatorCounter::SetUseCustomFormat(bool bInUseCustom)
+{
+	if (bUseCustomFormat == bInUseCustom)
+	{
+		return;
+	}
+
+	bUseCustomFormat = bInUseCustom;
+	OnUseCustomFormatChanged();
+}
+
+void UPropertyAnimatorCounter::SetPresetFormatName(FName InPresetName)
+{
+	if (InPresetName.IsEqual(PresetFormatName))
+	{
+		return;
+	}
+
+	if (!GetAvailableFormatNames().Contains(InPresetName))
+	{
+		return;
+	}
+
+	PresetFormatName = InPresetName;
+}
+
+void UPropertyAnimatorCounter::SetCustomFormat(const FPropertyAnimatorCounterFormat* InFormat)
+{
+	if (InFormat)
+	{
+		CustomFormat = TInstancedStruct<FPropertyAnimatorCounterFormat>::Make(*InFormat);
+	}
+	else
+	{
+		CustomFormat.Reset();
+	}
+
+	OnCustomFormatChanged();
+}
+
+const FPropertyAnimatorCounterFormat* UPropertyAnimatorCounter::GetCustomFormat() const
+{
+	return CustomFormat.GetPtr<FPropertyAnimatorCounterFormat>();
+}
+
+FString UPropertyAnimatorCounter::FormatNumber(double InNumber) const
+{
+	FText Output = FText::GetEmpty();
+
+	if (!bUseCustomFormat)
+	{
+		if (const UPropertyAnimatorSettings* AnimatorSettings = GetDefault<UPropertyAnimatorSettings>())
+		{
+			if (const FPropertyAnimatorCounterFormat* Format = AnimatorSettings->GetCounterFormat(PresetFormatName))
+			{
+				Output = FText::Format(DisplayPattern, FText::FromString(Format->FormatNumber(InNumber)));
+			}
+		}
+	}
+	else
+	{
+		if (const FPropertyAnimatorCounterFormat* Format = CustomFormat.GetPtr<FPropertyAnimatorCounterFormat>())
+		{
+			Output = FText::Format(DisplayPattern, FText::FromString(Format->FormatNumber(InNumber)));
+		}
+	}
+
+	return Output.ToString();
+}
+
+#if WITH_EDITOR
+void UPropertyAnimatorCounter::PostEditChangeProperty(FPropertyChangedEvent& InPropertyChangedEvent)
+{
+	Super::PostEditChangeProperty(InPropertyChangedEvent);
+
+	const FName MemberName = InPropertyChangedEvent.GetMemberPropertyName();
+
+	// PECP for instanced struct does not provide the correct MemberProperty
+	if (InPropertyChangedEvent.Property
+		&& InPropertyChangedEvent.Property->GetOwnerStruct() == FPropertyAnimatorCounterFormat::StaticStruct())
+	{
+		OnCustomFormatChanged();
+	}
+	else if (MemberName == GET_MEMBER_NAME_CHECKED(UPropertyAnimatorCounter, bUseCustomFormat))
+	{
+		OnUseCustomFormatChanged();
+	}
+}
+
+void UPropertyAnimatorCounter::OpenPropertyAnimatorSettings()
+{
+	if (const UPropertyAnimatorSettings* AnimatorSettings = GetDefault<UPropertyAnimatorSettings>())
+	{
+		AnimatorSettings->OpenSettings();
+	}
+}
+
+void UPropertyAnimatorCounter::SaveCustomFormatAsPreset()
+{
+	if (const FPropertyAnimatorCounterFormat* Format = CustomFormat.GetPtr<FPropertyAnimatorCounterFormat>())
+	{
+		if (UPropertyAnimatorSettings* AnimatorSettings = GetMutableDefault<UPropertyAnimatorSettings>())
+		{
+			if (AnimatorSettings->AddCounterFormat(*Format))
+			{
+				PresetFormatName = Format->FormatName;
+				bUseCustomFormat = false;
+				CustomFormat.Reset();
+			}
+		}
+	}
+}
+#endif
+
+EPropertyAnimatorPropertySupport UPropertyAnimatorCounter::IsPropertySupported(const FPropertyAnimatorCoreData& InPropertyData) const
+{
+	if (InPropertyData.IsA<FStrProperty>())
+	{
+		return EPropertyAnimatorPropertySupport::Complete;
+	}
+
+	// Check if a converter supports the conversion
+	if (UPropertyAnimatorCoreSubsystem* AnimatorSubsystem = UPropertyAnimatorCoreSubsystem::Get())
+	{
+		static const FPropertyBagPropertyDesc AnimatorTypeDesc("", EPropertyBagPropertyType::String);
+		const FPropertyBagPropertyDesc PropertyTypeDesc("", InPropertyData.GetLeafProperty());
+
+		if (AnimatorSubsystem->IsConversionSupported(AnimatorTypeDesc, PropertyTypeDesc))
+		{
+			return EPropertyAnimatorPropertySupport::Incomplete;
+		}
+	}
+
+	return Super::IsPropertySupported(InPropertyData);
+}
+
+void UPropertyAnimatorCounter::EvaluateProperties(FInstancedPropertyBag& InParameters)
+{
+	const double TimeElapsed = InParameters.GetValueDouble(TimeElapsedParameterName).GetValue();
+
+	EvaluateEachLinkedProperty<UPropertyAnimatorCoreContext>([this, TimeElapsed](
+		UPropertyAnimatorCoreContext* InContext
+		, const FPropertyAnimatorCoreData& InResolvedProperty
+		, FInstancedPropertyBag& InEvaluatedValues
+		, int32 InRangeIndex
+		, int32 InRangeMax)->bool
+	{
+		const FName DisplayName(InResolvedProperty.GetPathHash());
+
+		InEvaluatedValues.AddProperty(DisplayName, EPropertyBagPropertyType::String);
+		InEvaluatedValues.SetValueString(DisplayName, FormatNumber(TimeElapsed));
+
+		return true;
+	});
+}
+
+void UPropertyAnimatorCounter::OnPropertyLinked(UPropertyAnimatorCoreContext* InLinkedProperty, EPropertyAnimatorPropertySupport InSupport)
+{
+	Super::OnPropertyLinked(InLinkedProperty, InSupport);
+
+	if (EnumHasAnyFlags(InSupport, EPropertyAnimatorPropertySupport::Incomplete))
+	{
+		if (const UPropertyAnimatorCoreSubsystem* AnimatorSubsystem = UPropertyAnimatorCoreSubsystem::Get())
+		{
+			static const FPropertyBagPropertyDesc AnimatorTypeDesc("", EPropertyBagPropertyType::String);
+
+			const FPropertyAnimatorCoreData& Property = InLinkedProperty->GetAnimatedProperty();
+			const FPropertyBagPropertyDesc PropertyTypeDesc("", Property.GetLeafProperty());
+			const TSet<UPropertyAnimatorCoreConverterBase*> Converters = AnimatorSubsystem->GetSupportedConverters(AnimatorTypeDesc, PropertyTypeDesc);
+			check(!Converters.IsEmpty())
+
+			InLinkedProperty->SetConverterClass(Converters.Array()[0]->GetClass());
+		}
+	}
+}
+
+TArray<FName> UPropertyAnimatorCounter::GetAvailableFormatNames() const
+{
+	TArray<FName> FormatNames;
+
+	if (const UPropertyAnimatorSettings* AnimatorSettings = GetDefault<UPropertyAnimatorSettings>())
+	{
+		FormatNames = AnimatorSettings->GetCounterFormatNames().Array();
+	}
+
+	return FormatNames;
+}
+
+void UPropertyAnimatorCounter::OnCustomFormatChanged()
+{
+	if (FPropertyAnimatorCounterFormat* Format = CustomFormat.GetMutablePtr<FPropertyAnimatorCounterFormat>())
+	{
+		Format->EnsureCharactersLength();
+	}
+}
+
+void UPropertyAnimatorCounter::OnUseCustomFormatChanged()
+{
+	// Start from the current selected preset format
+	if (bUseCustomFormat && !CustomFormat.IsValid())
+	{
+		if (const UPropertyAnimatorSettings* AnimatorSettings = GetDefault<UPropertyAnimatorSettings>())
+		{
+			CustomFormat = TInstancedStruct<FPropertyAnimatorCounterFormat>::Make(*AnimatorSettings->GetCounterFormat(PresetFormatName));
+		}
+	}
+}
diff --git a/Engine/Plugins/Experimental/PropertyAnimator/Source/PropertyAnimator/Private/Settings/PropertyAnimatorSettings.cpp b/Engine/Plugins/Experimental/PropertyAnimator/Source/PropertyAnimator/Private/Settings/PropertyAnimatorSettings.cpp
new file mode 100644
--- /dev/null
+++ b/Engine/Plugins/Experimental/PropertyAnimator/Source/PropertyAnimator/Private/Settings/PropertyAnimatorSettings.cpp
@@ -0,0 +1,111 @@
+// Copyright Epic Games, Inc. All Rights Reserved.
+
+#include "Settings/PropertyAnimatorSettings.h"
+
+#if WITH_EDITOR
+#include "ISettingsModule.h"
+#include "Modules/ModuleManager.h"
+#endif
+
+UPropertyAnimatorSettings::UPropertyAnimatorSettings()
+{
+	CategoryName = TEXT("Motion Design");
+	SectionName = TEXT("Property Animator");
+
+	FPropertyAnimatorCounterFormat DefaultFormat;
+	DefaultFormat.FormatName = TEXT("Default");
+	DefaultFormat.bTruncate = false;
+	DefaultFormat.bUseSign = false;
+	DefaultFormat.DecimalCharacter = TEXT(".");
+	DefaultFormat.PaddingCharacter = TEXT("0");
+	DefaultFormat.GroupingCharacter = TEXT(",");
+	DefaultFormat.RoundingMode = EPropertyAnimatorCounterRoundingMode::None;
+	DefaultFormat.MaxDecimalCount = 3;
+	DefaultFormat.MinIntegerCount = 7;
+	DefaultFormat.GroupingSize = 3;
+
+	CounterFormatPresets.Add(DefaultFormat);
+}
+
+TSet<FName> UPropertyAnimatorSettings::GetCounterFormatNames() const
+{
+	TSet<FName> FormatNames;
+	Algo::Transform(CounterFormatPresets, FormatNames, [](const FPropertyAnimatorCounterFormat& InFormat)
+	{
+		return InFormat.FormatName;
+	});
+	return FormatNames;
+}
+
+const FPropertyAnimatorCounterFormat* UPropertyAnimatorSettings::GetCounterFormat(FName InName) const
+{
+	if (const FPropertyAnimatorCounterFormat* Format = CounterFormatPresets.Find(FPropertyAnimatorCounterFormat(InName)))
+	{
+		return Format;
+	}
+
+	return nullptr;
+}
+
+#if WITH_EDITOR
+bool UPropertyAnimatorSettings::AddCounterFormat(const FPropertyAnimatorCounterFormat& InNewFormat, bool bInOverride, bool bInSaveConfig)
+{
+	if (InNewFormat.FormatName.IsNone())
+	{
+		return false;
+	}
+
+	if (bInOverride)
+	{
+		CounterFormatPresets.Remove(InNewFormat);
+	}
+
+	bool bAlreadyInSet = false;
+	FPropertyAnimatorCounterFormat& NewFormat = CounterFormatPresets.FindOrAdd(InNewFormat, &bAlreadyInSet);
+
+	NewFormat.EnsureCharactersLength();
+
+	if (!bAlreadyInSet)
+	{
+		if (bInSaveConfig)
+		{
+			SaveConfig();
+			TryUpdateDefaultConfigFile();
+		}
+
+		return true;
+	}
+
+	return false;
+}
+
+void UPropertyAnimatorSettings::OpenSettings() const
+{
+	if (ISettingsModule* SettingsModule = FModuleManager::GetModulePtr<ISettingsModule>(TEXT("Settings")))
+	{
+		SettingsModule->ShowViewer(GetContainerName(), GetCategoryName(), GetSectionName());
+	}
+}
+
+void UPropertyAnimatorSettings::PostEditChangeProperty(FPropertyChangedEvent& InPropertyChangedEvent)
+{
+	Super::PostEditChangeProperty(InPropertyChangedEvent);
+
+	const FName MemberName = InPropertyChangedEvent.GetMemberPropertyName();
+
+	if (MemberName == GET_MEMBER_NAME_CHECKED(UPropertyAnimatorSettings, CounterFormatPresets)
+		&& InPropertyChangedEvent.ChangeType != EPropertyChangeType::Interactive)
+	{
+		OnCounterFormatsChanged();
+	}
+}
+#endif
+
+void UPropertyAnimatorSettings::OnCounterFormatsChanged()
+{
+	for (FPropertyAnimatorCounterFormat& Format : CounterFormatPresets)
+	{
+		Format.EnsureCharactersLength();
+	}
+}
+
diff --git a/Engine/Plugins/Experimental/PropertyAnimator/Source/PropertyAnimator/Private/Settings/PropertyAnimatorSettings.h b/Engine/Plugins/Experimental/PropertyAnimator/Source/PropertyAnimator/Private/Settings/PropertyAnimatorSettings.h
new file mode 100644
--- /dev/null
+++ b/Engine/Plugins/Experimental/PropertyAnimator/Source/PropertyAnimator/Private/Settings/PropertyAnimatorSettings.h
@@ -0,0 +1,37 @@
+// Copyright Epic Games, Inc. All Rights Reserved.
+
+#pragma once
+
+#include "Engine/DeveloperSettings.h"
+#include "Animators/PropertyAnimatorCounter.h"
+#include "PropertyAnimatorSettings.generated.h"
+
+/** Settings for motion design PropertyAnimator plugin */
+UCLASS(Config=Engine, meta=(DisplayName="Property Animator"))
+class UPropertyAnimatorSettings : public UDeveloperSettings
+{
+	GENERATED_BODY()
+
+public:
+	UPropertyAnimatorSettings();
+
+	TSet<FName> GetCounterFormatNames() const;
+	const FPropertyAnimatorCounterFormat* GetCounterFormat(FName InName) const;
+
+#if WITH_EDITOR
+	bool AddCounterFormat(const FPropertyAnimatorCounterFormat& InNewFormat, bool bInOverride = true, bool bInSaveConfig = true);
+	void OpenSettings() const;
+#endif
+
+protected:
+	//~ Begin UObject
+#if WITH_EDITOR
+	virtual void PostEditChangeProperty(FPropertyChangedEvent& InPropertyChangedEvent) override;
+#endif
+	//~ End UObject
+
+	void OnCounterFormatsChanged();
+
+	UPROPERTY(Config, EditAnywhere, Category="Animator", meta=(TitleProperty="FormatName"))
+	TSet<FPropertyAnimatorCounterFormat> CounterFormatPresets;
+};
diff --git a/Engine/Plugins/Experimental/PropertyAnimator/Source/PropertyAnimator/PropertyAnimator.Build.cs b/Engine/Plugins/Experimental/PropertyAnimator/Source/PropertyAnimator/PropertyAnimator.Build.cs
--- a/Engine/Plugins/Experimental/PropertyAnimator/Source/PropertyAnimator/PropertyAnimator.Build.cs
+++ b/Engine/Plugins/Experimental/PropertyAnimator/Source/PropertyAnimator/PropertyAnimator.Build.cs
@@ -22,6 +22,7 @@ public class PropertyAnimator : ModuleRules
 				"ApplicationCore",
 				"AudioSynesthesia",
 				"CoreUObject",
+				"DeveloperSettings",
 				"Engine",
 				"MovieSceneTracks",
 				"Text3D"
diff --git a/Engine/Plugins/Experimental/PropertyAnimator/Source/PropertyAnimator/Public/Animators/PropertyAnimatorCounter.h b/Engine/Plugins/Experimental/PropertyAnimator/Source/PropertyAnimator/Public/Animators/PropertyAnimatorCounter.h
new file mode 100644
--- /dev/null
+++ b/Engine/Plugins/Experimental/PropertyAnimator/Source/PropertyAnimator/Public/Animators/PropertyAnimatorCounter.h
@@ -0,0 +1,169 @@
+// Copyright Epic Games, Inc. All Rights Reserved.
+
+#pragma once
+
+#include "Animators/PropertyAnimatorNumericBase.h"
+#include "StructUtils/InstancedStruct.h"
+#include "PropertyAnimatorCounter.generated.h"
+
+/** Enumerates all rounding mode available */
+UENUM(BlueprintType)
+enum class EPropertyAnimatorCounterRoundingMode : uint8
+{
+	None,
+	Round,
+	Floor,
+	Ceil
+};
+
+/** Format options used to convert a number to string */
+USTRUCT(BlueprintType)
+struct FPropertyAnimatorCounterFormat
+{
+	GENERATED_BODY()
+
+	FPropertyAnimatorCounterFormat() = default;
+
+	explicit FPropertyAnimatorCounterFormat(FName InFormatName)
+		: FormatName(InFormatName)
+	{}
+
+	/** Format friendly name */
+	UPROPERTY(EditAnywhere, Category="Animator")
+	FName FormatName;
+
+	/** Minimum number of integer before the decimal separator for padding */
+	UPROPERTY(EditAnywhere, Category="Animator")
+	uint8 MinIntegerCount = 0;
+
+	/** Maximum number of decimal precision after the decimal separator */
+	UPROPERTY(EditAnywhere, Category="Animator")
+	uint8 MaxDecimalCount = 3;
+
+	/** Used to group numbers together like thousands */
+	UPROPERTY(EditAnywhere, Category="Animator")
+	uint8 GroupingSize = 3;
+
+	/** Decimal separator character */
+	UPROPERTY(EditAnywhere, Category="Animator")
+	FString DecimalCharacter = TEXT(".");
+
+	/** Thousands separator character */
+	UPROPERTY(EditAnywhere, Category="Animator")
+	FString GroupingCharacter = TEXT(",");
+
+	/** Filling character for leading blanks */
+	UPROPERTY(EditAnywhere, Category="Animator")
+	FString PaddingCharacter = TEXT("0");
+
+	/** Whether rounding the number is needed */
+	UPROPERTY(EditAnywhere, Category="Animator")
+	EPropertyAnimatorCounterRoundingMode RoundingMode = EPropertyAnimatorCounterRoundingMode::None;
+
+	/** Add a prefix symbol to show the sign of the number (+, -) */
+	UPROPERTY(EditAnywhere, Category="Animator")
+	bool bUseSign = false;
+
+	/** Truncate when the value exceeds the display format */
+	UPROPERTY(EditAnywhere, Category="Animator")
+	bool bTruncate = false;
+
+	/** Clamp custom characters to one char only */
+	void EnsureCharactersLength();
+
+	/** Format a number using these options */
+	FString FormatNumber(double InNumber) const;
+
+	friend uint32 GetTypeHash(const FPropertyAnimatorCounterFormat& InItem)
+	{
+		return GetTypeHash(InItem.FormatName);
+	}
+
+	bool operator==(const FPropertyAnimatorCounterFormat& InOther) const
+	{
+		return FormatName.IsEqual(InOther.FormatName);
+	}
+
+	bool operator!=(const FPropertyAnimatorCounterFormat& InOther) const
+	{
+		return !FormatName.IsEqual(InOther.FormatName);
+	}
+};
+
+/** Animate supported string properties to display a counter */
+UCLASS(MinimalAPI, AutoExpandCategories=("Animator"))
+class UPropertyAnimatorCounter : public UPropertyAnimatorCoreBase
+{
+	GENERATED_BODY()
+
+public:
+	static constexpr const TCHAR* DefaultAnimatorName = TEXT("Counter");
+
+	UPropertyAnimatorCounter();
+
+	void SetDisplayPattern(const FText& InPattern);
+	const FText& GetDisplayPattern() const
+	{
+		return DisplayPattern;
+	}
+
+	void SetUseCustomFormat(bool bInUseCustom);
+	bool GetUseCustomFormat() const
+	{
+		return bUseCustomFormat;
+	}
+
+	void SetPresetFormatName(FName InPresetName);
+	FName GetPresetFormatName() const
+	{
+		return PresetFormatName;
+	}
+
+	void SetCustomFormat(const FPropertyAnimatorCounterFormat* InFormat);
+	const FPropertyAnimatorCounterFormat* GetCustomFormat() const;
+
+protected:
+	FString FormatNumber(double InNumber) const;
+
+	//~ Begin UObject
+#if WITH_EDITOR
+	virtual void PostEditChangeProperty(FPropertyChangedEvent& InPropertyChangedEvent) override;
+#endif
+	//~ End UObject
+
+	//~ Begin UPropertyAnimatorCoreBase
+	virtual EPropertyAnimatorPropertySupport IsPropertySupported(const FPropertyAnimatorCoreData& InPropertyData) const override;
+	virtual void EvaluateProperties(FInstancedPropertyBag& InParameters) override;
+	virtual void OnPropertyLinked(UPropertyAnimatorCoreContext* InLinkedProperty, EPropertyAnimatorPropertySupport InSupport) override;
+	//~ End UPropertyAnimatorCoreBase
+
+	UFUNCTION()
+	TArray<FName> GetAvailableFormatNames() const;
+
+#if WITH_EDITOR
+	UFUNCTION(CallInEditor)
+	void OpenPropertyAnimatorSettings();
+
+	UFUNCTION(CallInEditor)
+	void SaveCustomFormatAsPreset();
+#endif
+
+	void OnUseCustomFormatChanged();
+	void OnCustomFormatChanged();
+
+	/** Display pattern for the output, use {0} as replacement symbol */
+	UPROPERTY(EditInstanceOnly, Setter, Getter, Category="Animator")
+	FText DisplayPattern = FText::FromString(TEXT("{0}"));
+
+	/** Use available presets formats or custom format */
+	UPROPERTY(EditInstanceOnly, Setter="SetUseCustomFormat", Getter="GetUseCustomFormat", Category="Animator")
+	bool bUseCustomFormat = false;
+
+	/** Preset format defined in the project settings */
+	UPROPERTY(EditInstanceOnly, Setter, Getter, Category="Animator", meta=(GetOptions="GetAvailableFormatNames", EditCondition="!bUseCustomFormat", EditConditionHides, HideEditConditionToggle))
+	FName PresetFormatName;
+
+	/** Custom format */
+	UPROPERTY(EditInstanceOnly, Category="Animator", meta=(EditCondition="bUseCustomFormat", EditConditionHides, HideEditConditionToggle))
+	TInstancedStruct<FPropertyAnimatorCounterFormat> CustomFormat;
+};
\ No newline at end of file
