MotionDesign : PropertyAnimator
- Fixed counter animator custom format instanced struct property edit condition refresh

diff --git a/Engine/Plugins/Experimental/PropertyAnimator/Source/PropertyAnimator/Private/Animators/PropertyAnimatorCounter.cpp b/Engine/Plugins/Experimental/PropertyAnimator/Source/PropertyAnimator/Private/Animators/PropertyAnimatorCounter.cpp
--- a/Engine/Plugins/Experimental/PropertyAnimator/Source/PropertyAnimator/Private/Animators/PropertyAnimatorCounter.cpp
+++ b/Engine/Plugins/Experimental/PropertyAnimator/Source/PropertyAnimator/Private/Animators/PropertyAnimatorCounter.cpp
@@ -206,6 +206,11 @@ FString UPropertyAnimatorCounter::FormatNumber(double InNumber) const
 }
 
 #if WITH_EDITOR
+FName UPropertyAnimatorCounter::GetUseCustomFormatPropertyName()
+{
+	return GET_MEMBER_NAME_CHECKED(UPropertyAnimatorCounter, bUseCustomFormat);
+}
+
 void UPropertyAnimatorCounter::PostEditChangeProperty(FPropertyChangedEvent& InPropertyChangedEvent)
 {
 	Super::PostEditChangeProperty(InPropertyChangedEvent);
diff --git a/Engine/Plugins/Experimental/PropertyAnimator/Source/PropertyAnimator/Public/Animators/PropertyAnimatorCounter.h b/Engine/Plugins/Experimental/PropertyAnimator/Source/PropertyAnimator/Public/Animators/PropertyAnimatorCounter.h
--- a/Engine/Plugins/Experimental/PropertyAnimator/Source/PropertyAnimator/Public/Animators/PropertyAnimatorCounter.h
+++ b/Engine/Plugins/Experimental/PropertyAnimator/Source/PropertyAnimator/Public/Animators/PropertyAnimatorCounter.h
@@ -99,6 +99,10 @@ class UPropertyAnimatorCounter : public UPropertyAnimatorCoreBase
 public:
 	static constexpr const TCHAR* DefaultAnimatorName = TEXT("Counter");
 
+#if WITH_EDITOR
+	PROPERTYANIMATOR_API static FName GetUseCustomFormatPropertyName();
+#endif
+
 	UPropertyAnimatorCounter();
 
 	void SetDisplayPattern(const FText& InPattern);
diff --git a/Engine/Plugins/Experimental/PropertyAnimator/Source/PropertyAnimatorEditor/Private/Customizations/PropertyAnimatorEditorCounterDetailCustomization.cpp b/Engine/Plugins/Experimental/PropertyAnimator/Source/PropertyAnimatorEditor/Private/Customizations/PropertyAnimatorEditorCounterDetailCustomization.cpp
new file mode 100644
--- /dev/null
+++ b/Engine/Plugins/Experimental/PropertyAnimator/Source/PropertyAnimatorEditor/Private/Customizations/PropertyAnimatorEditorCounterDetailCustomization.cpp
@@ -0,0 +1,33 @@
+// Copyright Epic Games, Inc. All Rights Reserved.
+
+#include "Customizations/PropertyAnimatorEditorCounterDetailCustomization.h"
+
+#include "Animators/PropertyAnimatorCounter.h"
+#include "DetailLayoutBuilder.h"
+
+void FPropertyAnimatorEditorCounterDetailCustomization::CustomizeDetails(IDetailLayoutBuilder& InDetailBuilder)
+{
+	UseCustomFormatHandle = InDetailBuilder.GetProperty(
+		UPropertyAnimatorCounter::GetUseCustomFormatPropertyName(),
+		UPropertyAnimatorCounter::StaticClass()
+	);
+
+	if (!UseCustomFormatHandle->IsValidHandle())
+	{
+		return;
+	}
+
+	// Due to instanced struct EditCondition not working and needing a refresh
+	UseCustomFormatHandle->SetOnPropertyValueChanged(FSimpleDelegate::CreateSP(this, &FPropertyAnimatorEditorCounterDetailCustomization::OnConditionPropertyChanged));
+}
+
+void FPropertyAnimatorEditorCounterDetailCustomization::OnConditionPropertyChanged() const
+{
+	if (UseCustomFormatHandle.IsValid() && UseCustomFormatHandle->IsValidHandle())
+	{
+		if (const TSharedPtr<IPropertyHandle> ParentPropertyHandle = UseCustomFormatHandle->GetParentHandle())
+		{
+			ParentPropertyHandle->RequestRebuildChildren();
+		}
+	}
+}
diff --git a/Engine/Plugins/Experimental/PropertyAnimator/Source/PropertyAnimatorEditor/Private/Customizations/PropertyAnimatorEditorCounterDetailCustomization.h b/Engine/Plugins/Experimental/PropertyAnimator/Source/PropertyAnimatorEditor/Private/Customizations/PropertyAnimatorEditorCounterDetailCustomization.h
new file mode 100644
--- /dev/null
+++ b/Engine/Plugins/Experimental/PropertyAnimator/Source/PropertyAnimatorEditor/Private/Customizations/PropertyAnimatorEditorCounterDetailCustomization.h
@@ -0,0 +1,27 @@
+// Copyright Epic Games, Inc. All Rights Reserved.
+
+#pragma once
+
+#include "IDetailCustomization.h"
+#include "Templates/SharedPointer.h"
+
+class IPropertyHandle;
+
+/** Details customization for UPropertyAnimatorCounter */
+class FPropertyAnimatorEditorCounterDetailCustomization : public IDetailCustomization
+{
+public:
+	static TSharedRef<IDetailCustomization> MakeInstance()
+	{
+		return MakeShared<FPropertyAnimatorEditorCounterDetailCustomization>();
+	}
+
+	//~ Begin IDetailCustomization
+	virtual void CustomizeDetails(IDetailLayoutBuilder& InDetailBuilder) override;
+	//~ End IDetailCustomization
+
+protected:
+	void OnConditionPropertyChanged() const;
+
+	TSharedPtr<IPropertyHandle> UseCustomFormatHandle;
+};
diff --git a/Engine/Plugins/Experimental/PropertyAnimator/Source/PropertyAnimatorEditor/Private/PropertyAnimatorEditorModule.cpp b/Engine/Plugins/Experimental/PropertyAnimator/Source/PropertyAnimatorEditor/Private/PropertyAnimatorEditorModule.cpp
--- a/Engine/Plugins/Experimental/PropertyAnimator/Source/PropertyAnimatorEditor/Private/PropertyAnimatorEditorModule.cpp
+++ b/Engine/Plugins/Experimental/PropertyAnimator/Source/PropertyAnimatorEditor/Private/PropertyAnimatorEditorModule.cpp
@@ -2,8 +2,12 @@
 
 #include "PropertyAnimatorEditorModule.h"
 
+#include "Animators/PropertyAnimatorCounter.h"
+#include "Customizations/PropertyAnimatorEditorCounterDetailCustomization.h"
+#include "Modules/ModuleManager.h"
 #include "MovieScene/Easing/PropertyAnimatorEasingDoubleChannel.h"
 #include "MovieScene/Wave/PropertyAnimatorWaveDoubleChannel.h"
+#include "PropertyEditorModule.h"
 #include "Styles/PropertyAnimatorEditorStyle.h"
 
 void FPropertyAnimatorEditorModule::StartupModule()
@@ -14,6 +18,19 @@ void FPropertyAnimatorEditorModule::StartupModule()
 
 	// Init once
 	FPropertyAnimatorEditorStyle::Get();
+
+	FPropertyEditorModule& PropertyModule = FModuleManager::LoadModuleChecked<FPropertyEditorModule>("PropertyEditor");
+	PropertyModule.RegisterCustomClassLayout(UPropertyAnimatorCounter::StaticClass()->GetFName(), FOnGetDetailCustomizationInstance::CreateStatic(&FPropertyAnimatorEditorCounterDetailCustomization::MakeInstance));
+
+}
+
+void FPropertyAnimatorEditorModule::ShutdownModule()
+{
+	if (FModuleManager::Get().IsModuleLoaded("PropertyEditor") && UObjectInitialized())
+	{
+		FPropertyEditorModule& PropertyModule = FModuleManager::GetModuleChecked<FPropertyEditorModule>("PropertyEditor");
+		PropertyModule.UnregisterCustomClassLayout(UPropertyAnimatorCounter::StaticClass()->GetFName());
+	}
 }
 
 IMPLEMENT_MODULE(FPropertyAnimatorEditorModule, PropertyAnimatorEditor)
diff --git a/Engine/Plugins/Experimental/PropertyAnimator/Source/PropertyAnimatorEditor/Private/PropertyAnimatorEditorModule.h b/Engine/Plugins/Experimental/PropertyAnimator/Source/PropertyAnimatorEditor/Private/PropertyAnimatorEditorModule.h
--- a/Engine/Plugins/Experimental/PropertyAnimator/Source/PropertyAnimatorEditor/Private/PropertyAnimatorEditorModule.h
+++ b/Engine/Plugins/Experimental/PropertyAnimator/Source/PropertyAnimatorEditor/Private/PropertyAnimatorEditorModule.h
@@ -14,6 +14,7 @@ class FPropertyAnimatorEditorModule : public IModuleInterface
 public:
 	//~ Begin IModuleInterface
 	virtual void StartupModule() override;
+	virtual void ShutdownModule() override;
 	//~ End IModuleInterface
 
 	template<typename InCurveChannelType>
