MotionDesign : PropertyAnimator
- Disabling cycle options within soundwave animator and restoring previous code

diff --git a/Engine/Plugins/Experimental/PropertyAnimator/Source/PropertyAnimator/Private/Animators/PropertyAnimatorSoundWave.cpp b/Engine/Plugins/Experimental/PropertyAnimator/Source/PropertyAnimator/Private/Animators/PropertyAnimatorSoundWave.cpp
--- a/Engine/Plugins/Experimental/PropertyAnimator/Source/PropertyAnimator/Private/Animators/PropertyAnimatorSoundWave.cpp
+++ b/Engine/Plugins/Experimental/PropertyAnimator/Source/PropertyAnimator/Private/Animators/PropertyAnimatorSoundWave.cpp
@@ -10,6 +10,7 @@
 UPropertyAnimatorSoundWave::UPropertyAnimatorSoundWave()
 {
 	SetAnimatorDisplayName(DefaultControllerName);
+	CycleMode = EPropertyAnimatorCycleMode::None;
 }
 
 void UPropertyAnimatorSoundWave::SetSampledSoundWave(USoundWave* InSoundWave)
@@ -63,23 +64,19 @@ bool UPropertyAnimatorSoundWave::EvaluateProperty(const FPropertyAnimatorCoreDat
 {
 	if (AudioAnalyzer && AudioAnalyzer->DurationInSeconds > 0)
 	{
-		double TimeElapsed = InParameters.GetValueDouble(TimeElapsedParameterName).GetValue();
-		float Frequency = InParameters.GetValueFloat(FrequencyParameterName).GetValue();
+		double SampleTime = InParameters.GetValueDouble(TimeElapsedParameterName).GetValue();
 
-		const float Period = 1.f / Frequency;
-		float SampleTime = FMath::Fmod(TimeElapsed, Period);
-
-		if (FMath::Abs(TimeElapsed / Period) <= 1.f || bLoop)
+		if ((SampleTime >= 0 && SampleTime <= AudioAnalyzer->DurationInSeconds) || bLoop)
 		{
+			SampleTime = FMath::Fmod(SampleTime, AudioAnalyzer->DurationInSeconds);
+
 			if (SampleTime < 0)
 			{
-				SampleTime = Period - FMath::Abs(SampleTime);
+				SampleTime = AudioAnalyzer->DurationInSeconds + SampleTime;
 			}
 
-			const float NormalizedSampleTime = FMath::GetMappedRangeValueClamped(FVector2D(0, Period), FVector2D(0, AudioAnalyzer->DurationInSeconds), SampleTime);
-
 			float NormalizedLoudness = 0.f;
-			AudioAnalyzer->GetNormalizedLoudnessAtTime(NormalizedSampleTime, NormalizedLoudness);
+			AudioAnalyzer->GetNormalizedLoudnessAtTime(SampleTime, NormalizedLoudness);
 
 			InParameters.AddProperty(AlphaParameterName, EPropertyBagPropertyType::Float);
 			InParameters.SetValueFloat(AlphaParameterName, NormalizedLoudness);
