MotionDesign : PropertyAnimator
- Removed property group system from animator
- Added range selector based on custom resolver (text) with preset import/export

diff --git a/Engine/Plugins/Experimental/PropertyAnimator/Source/PropertyAnimator/Private/Animators/PropertyAnimatorCounter.cpp b/Engine/Plugins/Experimental/PropertyAnimator/Source/PropertyAnimator/Private/Animators/PropertyAnimatorCounter.cpp
--- a/Engine/Plugins/Experimental/PropertyAnimator/Source/PropertyAnimator/Private/Animators/PropertyAnimatorCounter.cpp
+++ b/Engine/Plugins/Experimental/PropertyAnimator/Source/PropertyAnimator/Private/Animators/PropertyAnimatorCounter.cpp
@@ -87,7 +87,7 @@ FString FPropertyAnimatorCounterFormat::FormatNumber(double InNumber) const
 		IntegerPart.InsertAt(0, FString::ChrN(IntegerCount - IntegerPart.Len(), PaddingCharacter[0]));
 	}
 
-	if (!GroupingCharacter.IsEmpty())
+	if (!GroupingCharacter.IsEmpty() && GroupingSize > 0)
 	{
 		int32 ThousandsIndex = IntegerPart.Len() - 1;
 		int32 ThousandsCount = 0;
diff --git a/Engine/Plugins/Experimental/PropertyAnimator/Source/PropertyAnimator/Private/Properties/PropertyAnimatorTextGroup.cpp b/Engine/Plugins/Experimental/PropertyAnimator/Source/PropertyAnimator/Private/Properties/PropertyAnimatorTextGroup.cpp
deleted file mode 100644
diff --git a/Engine/Plugins/Experimental/PropertyAnimator/Source/PropertyAnimator/Private/Properties/PropertyAnimatorTextResolver.cpp b/Engine/Plugins/Experimental/PropertyAnimator/Source/PropertyAnimator/Private/Properties/PropertyAnimatorTextResolver.cpp
--- a/Engine/Plugins/Experimental/PropertyAnimator/Source/PropertyAnimator/Private/Properties/PropertyAnimatorTextResolver.cpp
+++ b/Engine/Plugins/Experimental/PropertyAnimator/Source/PropertyAnimator/Private/Properties/PropertyAnimatorTextResolver.cpp
@@ -2,8 +2,65 @@
 
 #include "Properties/PropertyAnimatorTextResolver.h"
 
+#include "Dom/JsonObject.h"
+#include "Dom/JsonValue.h"
 #include "Text3DComponent.h"
 
+void UPropertyAnimatorTextResolver::SetUnit(EPropertyAnimatorTextResolverRangeUnit InUnit)
+{
+	Unit = InUnit;
+}
+
+void UPropertyAnimatorTextResolver::SetStart(float InRangeStart)
+{
+	Start = FMath::Clamp(InRangeStart, 0, 100);
+}
+
+void UPropertyAnimatorTextResolver::SetEnd(float InRangeEnd)
+{
+	End = FMath::Clamp(InRangeEnd, 0, 100);
+}
+
+void UPropertyAnimatorTextResolver::SetOffset(float InRangeOffset)
+{
+	Offset = InRangeOffset;
+}
+
+void UPropertyAnimatorTextResolver::SetCharacterStartIndex(int32 InRangeStart)
+{
+	CharacterStartIndex = FMath::Max(InRangeStart, 0);
+}
+
+void UPropertyAnimatorTextResolver::SetCharacterEndIndex(int32 InRangeEnd)
+{
+	CharacterEndIndex = FMath::Max(InRangeEnd, 0);
+}
+
+void UPropertyAnimatorTextResolver::SetCharacterOffsetIndex(int32 InRangeOffset)
+{
+	CharacterOffsetIndex = InRangeOffset;
+}
+
+void UPropertyAnimatorTextResolver::SetWordStartIndex(int32 InRangeStart)
+{
+	WordStartIndex = FMath::Max(InRangeStart, 0);
+}
+
+void UPropertyAnimatorTextResolver::SetWordEndIndex(int32 InRangeEnd)
+{
+	WordEndIndex = FMath::Max(InRangeEnd, 0);
+}
+
+void UPropertyAnimatorTextResolver::SetWordOffsetIndex(int32 InRangeOffset)
+{
+	WordOffsetIndex = InRangeOffset;
+}
+
+void UPropertyAnimatorTextResolver::SetDirection(EPropertyAnimatorTextResolverRangeDirection InDirection)
+{
+	Direction = InDirection;
+}
+
 void UPropertyAnimatorTextResolver::GetResolvableProperties(const FPropertyAnimatorCoreData& InParentProperty, TSet<FPropertyAnimatorCoreData>& OutProperties)
 {
 	const AActor* Actor = InParentProperty.GetOwningActor();
@@ -37,7 +94,7 @@ void UPropertyAnimatorTextResolver::GetResolvableProperties(const FPropertyAnima
 	OutProperties.Add(RelativeScaleProperty);
 }
 
-void UPropertyAnimatorTextResolver::ResolveProperties(const FPropertyAnimatorCoreData& InTemplateProperty, TArray<FPropertyAnimatorCoreData>& OutProperties)
+void UPropertyAnimatorTextResolver::ResolveProperties(const FPropertyAnimatorCoreData& InTemplateProperty, TArray<FPropertyAnimatorCoreData>& OutProperties, bool bInForEvaluation)
 {
 	if (!InTemplateProperty.IsResolvable())
 	{
@@ -65,4 +122,224 @@ void UPropertyAnimatorTextResolver::ResolveProperties(const FPropertyAnimatorCor
 		FPropertyAnimatorCoreData CharacterProperty(CharacterKerningComponent, ChainProperties);
 		OutProperties.Add(CharacterProperty);
 	}
+
+	if (!bInForEvaluation || OutProperties.IsEmpty())
+	{
+		return;
+	}
+
+	const int32 MaxIndex = OutProperties.Num();
+	int32 BeginIndex = 0;
+	int32 EndIndex = 0;
+
+	switch (Unit)
+	{
+		case EPropertyAnimatorTextResolverRangeUnit::Percentage:
+		{
+			float StartPercentage = Start / 100.f;
+			float EndPercentage = End / 100.f;
+			float OffsetPercentage = Offset / 100.f;
+
+			if (Direction == EPropertyAnimatorTextResolverRangeDirection::RightToLeft)
+			{
+				const float Temp = 1.f - StartPercentage;
+				StartPercentage = 1.f - EndPercentage;
+				EndPercentage = Temp;
+				OffsetPercentage *= -1;
+			}
+			else if (Direction == EPropertyAnimatorTextResolverRangeDirection::FromCenter)
+			{
+				constexpr float MidPercentage = 0.5f;
+				const float Expansion = EndPercentage / 2.f;
+				StartPercentage = MidPercentage - Expansion;
+				EndPercentage = MidPercentage + Expansion;
+			}
+
+			BeginIndex = StartPercentage * MaxIndex + OffsetPercentage * MaxIndex;
+			EndIndex = EndPercentage * MaxIndex + OffsetPercentage * MaxIndex;
+		}
+		break;
+
+		case EPropertyAnimatorTextResolverRangeUnit::Character:
+		{
+			int32 CharacterStart = CharacterStartIndex;
+			int32 CharacterEnd = CharacterEndIndex;
+			int32 CharacterOffset = CharacterOffsetIndex;
+
+			if (Direction == EPropertyAnimatorTextResolverRangeDirection::RightToLeft)
+			{
+				const int32 Temp = MaxIndex - CharacterStart;
+				CharacterStart = MaxIndex - CharacterEnd;
+				CharacterEnd = Temp;
+				CharacterOffset *= -1;
+			}
+			else if (Direction == EPropertyAnimatorTextResolverRangeDirection::FromCenter)
+			{
+				const int32 CharacterMid = MaxIndex / 2;
+				const int32 Expansion = CharacterEnd / 2;
+				CharacterStart = CharacterMid - Expansion;
+				CharacterEnd = CharacterMid + Expansion;
+			}
+
+			BeginIndex = CharacterStart + CharacterOffset;
+			EndIndex = CharacterEnd + CharacterOffset;
+		}
+		break;
+
+		case EPropertyAnimatorTextResolverRangeUnit::Word:
+		{
+			if (const UText3DComponent* TextComponent = TextRootComponent->GetTypedOuter<UText3DComponent>())
+			{
+				const FText3DStatistics& TextStats = TextComponent->GetStatistics();
+
+				if (TextStats.Words.IsEmpty())
+				{
+					break;
+				}
+
+				const int32 WordCount = TextStats.Words.Num();
+				int32 WordStart = WordStartIndex;
+				int32 WordEnd = WordEndIndex;
+				int32 WordOffset = WordOffsetIndex;
+
+				if (Direction == EPropertyAnimatorTextResolverRangeDirection::RightToLeft)
+				{
+					const int32 Temp = WordCount - WordStart;
+					WordStart = WordCount - WordEnd;
+					WordEnd = Temp;
+					WordOffset *= -1;
+				}
+				else if (Direction == EPropertyAnimatorTextResolverRangeDirection::FromCenter)
+				{
+					const int32 WordMid = FMath::CeilToInt(WordCount / 2.f);
+					const int32 Expansion = FMath::CeilToInt(WordEnd / 2.f);
+					WordStart = WordMid - Expansion;
+					WordEnd = WordMid + Expansion;
+				}
+
+				if (WordStart != WordEnd)
+				{
+					WordStart += WordOffset;
+					WordEnd += WordOffset - 1;
+
+					if (TextStats.Words.IsValidIndex(WordStart))
+					{
+						BeginIndex = TextStats.Words[WordStart].RenderRange.BeginIndex;
+					}
+
+					if (TextStats.Words.IsValidIndex(WordEnd))
+					{
+						EndIndex = TextStats.Words[WordEnd].RenderRange.EndIndex;
+					}
+					else if (WordEnd >= WordCount && WordStart < WordCount)
+					{
+						EndIndex = TextStats.Words.Last().RenderRange.EndIndex;
+					}
+				}
+			}
+		}
+		break;
+	}
+
+	if (EndIndex < 0 || BeginIndex > EndIndex || BeginIndex == EndIndex || BeginIndex > MaxIndex)
+	{
+		OutProperties.Empty();
+		return;
+	}
+
+	// Remove at the end
+	if (EndIndex < MaxIndex)
+	{
+		OutProperties.RemoveAt(EndIndex, MaxIndex - EndIndex);
+	}
+
+	// Remove at the start
+	if (BeginIndex > 0 && BeginIndex <= MaxIndex)
+	{
+		OutProperties.RemoveAt(0, BeginIndex);
+	}
+}
+
+bool UPropertyAnimatorTextResolver::ImportPreset(const UPropertyAnimatorCorePresetBase* InPreset, const TSharedRef<FJsonValue>& InValue)
+{
+	const TSharedPtr<FJsonObject>* JsonResolverObjectPtr;
+
+	if (Super::ImportPreset(InPreset, InValue) && InValue->TryGetObject(JsonResolverObjectPtr))
+	{
+		const TSharedPtr<FJsonObject> JsonResolverObject = *JsonResolverObjectPtr;
+
+		uint8 JsonUnit = static_cast<uint8>(Unit);
+		JsonResolverObject->TryGetNumberField(GET_MEMBER_NAME_STRING_CHECKED(UPropertyAnimatorTextResolver, Unit), JsonUnit);
+		SetUnit(static_cast<EPropertyAnimatorTextResolverRangeUnit>(JsonUnit));
+
+		float JsonStart = Start;
+		JsonResolverObject->TryGetNumberField(GET_MEMBER_NAME_STRING_CHECKED(UPropertyAnimatorTextResolver, Start), JsonStart);
+		SetStart(JsonStart);
+
+		float JsonEnd = End;
+		JsonResolverObject->TryGetNumberField(GET_MEMBER_NAME_STRING_CHECKED(UPropertyAnimatorTextResolver, End), JsonEnd);
+		SetEnd(JsonEnd);
+
+		float JsonOffset = Offset;
+		JsonResolverObject->TryGetNumberField(GET_MEMBER_NAME_STRING_CHECKED(UPropertyAnimatorTextResolver, Offset), JsonOffset);
+		SetOffset(JsonOffset);
+
+		int32 JsonCharStartIndex = CharacterStartIndex;
+		JsonResolverObject->TryGetNumberField(GET_MEMBER_NAME_STRING_CHECKED(UPropertyAnimatorTextResolver, CharacterStartIndex), JsonCharStartIndex);
+		SetCharacterStartIndex(JsonCharStartIndex);
+
+		int32 JsonCharEndIndex = CharacterEndIndex;
+		JsonResolverObject->TryGetNumberField(GET_MEMBER_NAME_STRING_CHECKED(UPropertyAnimatorTextResolver, CharacterEndIndex), JsonCharEndIndex);
+		SetCharacterEndIndex(JsonCharEndIndex);
+
+		int32 JsonCharOffsetIndex = CharacterOffsetIndex;
+		JsonResolverObject->TryGetNumberField(GET_MEMBER_NAME_STRING_CHECKED(UPropertyAnimatorTextResolver, CharacterOffsetIndex), JsonCharOffsetIndex);
+		SetCharacterOffsetIndex(JsonCharOffsetIndex);
+
+		int32 JsonWordStartIndex = WordStartIndex;
+		JsonResolverObject->TryGetNumberField(GET_MEMBER_NAME_STRING_CHECKED(UPropertyAnimatorTextResolver, WordStartIndex), JsonWordStartIndex);
+		SetWordStartIndex(JsonWordStartIndex);
+
+		int32 JsonWordEndIndex = WordEndIndex;
+		JsonResolverObject->TryGetNumberField(GET_MEMBER_NAME_STRING_CHECKED(UPropertyAnimatorTextResolver, WordEndIndex), JsonWordEndIndex);
+		SetWordEndIndex(JsonWordEndIndex);
+
+		int32 JsonWordOffsetIndex = WordOffsetIndex;
+		JsonResolverObject->TryGetNumberField(GET_MEMBER_NAME_STRING_CHECKED(UPropertyAnimatorTextResolver, WordOffsetIndex), JsonWordOffsetIndex);
+		SetWordOffsetIndex(JsonWordOffsetIndex);
+
+		uint8 JsonDirection = static_cast<uint8>(Direction);
+		JsonResolverObject->TryGetNumberField(GET_MEMBER_NAME_STRING_CHECKED(UPropertyAnimatorTextResolver, Direction), JsonDirection);
+		SetDirection(static_cast<EPropertyAnimatorTextResolverRangeDirection>(JsonWordOffsetIndex));
+
+		return true;
+	}
+
+	return false;
+}
+
+bool UPropertyAnimatorTextResolver::ExportPreset(const UPropertyAnimatorCorePresetBase* InPreset, TSharedPtr<FJsonValue>& OutValue)
+{
+	const TSharedPtr<FJsonObject>* JsonResolverObjectPtr;
+
+	if (Super::ExportPreset(InPreset, OutValue) && OutValue->TryGetObject(JsonResolverObjectPtr))
+	{
+		const TSharedPtr<FJsonObject> JsonResolverObject = *JsonResolverObjectPtr;
+
+		JsonResolverObject->SetNumberField(GET_MEMBER_NAME_STRING_CHECKED(UPropertyAnimatorTextResolver, Unit), static_cast<uint8>(Unit));
+		JsonResolverObject->SetNumberField(GET_MEMBER_NAME_STRING_CHECKED(UPropertyAnimatorTextResolver, Start), Start);
+		JsonResolverObject->SetNumberField(GET_MEMBER_NAME_STRING_CHECKED(UPropertyAnimatorTextResolver, End), End);
+		JsonResolverObject->SetNumberField(GET_MEMBER_NAME_STRING_CHECKED(UPropertyAnimatorTextResolver, Offset), Offset);
+		JsonResolverObject->SetNumberField(GET_MEMBER_NAME_STRING_CHECKED(UPropertyAnimatorTextResolver, CharacterStartIndex), CharacterStartIndex);
+		JsonResolverObject->SetNumberField(GET_MEMBER_NAME_STRING_CHECKED(UPropertyAnimatorTextResolver, CharacterEndIndex), CharacterEndIndex);
+		JsonResolverObject->SetNumberField(GET_MEMBER_NAME_STRING_CHECKED(UPropertyAnimatorTextResolver, CharacterOffsetIndex), CharacterOffsetIndex);
+		JsonResolverObject->SetNumberField(GET_MEMBER_NAME_STRING_CHECKED(UPropertyAnimatorTextResolver, WordStartIndex), WordStartIndex);
+		JsonResolverObject->SetNumberField(GET_MEMBER_NAME_STRING_CHECKED(UPropertyAnimatorTextResolver, WordEndIndex), WordEndIndex);
+		JsonResolverObject->SetNumberField(GET_MEMBER_NAME_STRING_CHECKED(UPropertyAnimatorTextResolver, WordOffsetIndex), WordOffsetIndex);
+		JsonResolverObject->SetNumberField(GET_MEMBER_NAME_STRING_CHECKED(UPropertyAnimatorTextResolver, Direction), static_cast<uint8>(Direction));
+
+		return true;
+	}
+
+	return false;
 }
diff --git a/Engine/Plugins/Experimental/PropertyAnimator/Source/PropertyAnimator/Private/Properties/PropertyAnimatorTextResolver.h b/Engine/Plugins/Experimental/PropertyAnimator/Source/PropertyAnimator/Private/Properties/PropertyAnimatorTextResolver.h
--- a/Engine/Plugins/Experimental/PropertyAnimator/Source/PropertyAnimator/Private/Properties/PropertyAnimatorTextResolver.h
+++ b/Engine/Plugins/Experimental/PropertyAnimator/Source/PropertyAnimator/Private/Properties/PropertyAnimatorTextResolver.h
@@ -6,24 +6,142 @@
 #include "Properties/PropertyAnimatorCoreResolver.h"
 #include "PropertyAnimatorTextResolver.generated.h"
 
+UENUM()
+enum class EPropertyAnimatorTextResolverRangeUnit : uint8
+{
+	Percentage,
+	Character,
+	Word
+};
+
+UENUM()
+enum class EPropertyAnimatorTextResolverRangeDirection : uint8
+{
+	LeftToRight,
+	RightToLeft,
+	FromCenter
+};
+
 /**
  * Text characters properties resolver
  * Since each character in text are transient and regenerated on change
- * We need to have a resolvable property that will be resolved to each characters in the text we needed
- * We manipulate that resolvable property that underneath means we manipulate all text characters properties
+ * We need to have a resolver that will resolve each character in the text when needed
+ * We manipulate a single property that underneath means we manipulate all text characters properties
  */
-UCLASS(Transient)
+UCLASS()
 class UPropertyAnimatorTextResolver : public UPropertyAnimatorCoreResolver
 {
 	GENERATED_BODY()
 
 public:
 	UPropertyAnimatorTextResolver()
-		: UPropertyAnimatorCoreResolver(TEXT("TextCharacters"), /** IsRange */ true)
+		: UPropertyAnimatorCoreResolver(TEXT("TextChars"))
 	{}
 
+	PROPERTYANIMATOR_API void SetUnit(EPropertyAnimatorTextResolverRangeUnit InUnit);
+	EPropertyAnimatorTextResolverRangeUnit GetUnit() const
+	{
+		return Unit;
+	}
+
+	PROPERTYANIMATOR_API void SetStart(float InRangeStart);
+	float GetStart() const
+	{
+		return Start;
+	}
+
+	PROPERTYANIMATOR_API void SetEnd(float InRangeEnd);
+	float GetEnd() const
+	{
+		return End;
+	}
+
+	PROPERTYANIMATOR_API void SetOffset(float InRangeOffset);
+	float GetOffset() const
+	{
+		return Offset;
+	}
+
+	PROPERTYANIMATOR_API void SetCharacterStartIndex(int32 InRangeStart);
+	int32 GetCharacterStartIndex() const
+	{
+		return CharacterStartIndex;
+	}
+
+	PROPERTYANIMATOR_API void SetCharacterEndIndex(int32 InRangeEnd);
+	int32 GetCharacterEndIndex() const
+	{
+		return CharacterEndIndex;
+	}
+
+	PROPERTYANIMATOR_API void SetCharacterOffsetIndex(int32 InRangeOffset);
+	int32 GetCharacterOffsetIndex() const
+	{
+		return CharacterOffsetIndex;
+	}
+
+	PROPERTYANIMATOR_API void SetWordStartIndex(int32 InRangeStart);
+	int32 GetWordStartIndex() const
+	{
+		return WordStartIndex;
+	}
+
+	PROPERTYANIMATOR_API void SetWordEndIndex(int32 InRangeEnd);
+	int32 GetWordEndIndex() const
+	{
+		return WordEndIndex;
+	}
+
+	PROPERTYANIMATOR_API void SetWordOffsetIndex(int32 InRangeOffset);
+	int32 GetWordOffsetIndex() const
+	{
+		return WordOffsetIndex;
+	}
+
+	PROPERTYANIMATOR_API void SetDirection(EPropertyAnimatorTextResolverRangeDirection InDirection);
+	EPropertyAnimatorTextResolverRangeDirection GeDirection() const
+	{
+		return Direction;
+	}
+
 	//~ Begin UPropertyAnimatorCoreResolver
 	virtual void GetResolvableProperties(const FPropertyAnimatorCoreData& InParentProperty, TSet<FPropertyAnimatorCoreData>& OutProperties) override;
-	virtual void ResolveProperties(const FPropertyAnimatorCoreData& InTemplateProperty, TArray<FPropertyAnimatorCoreData>& OutProperties) override;
+	virtual void ResolveProperties(const FPropertyAnimatorCoreData& InTemplateProperty, TArray<FPropertyAnimatorCoreData>& OutProperties, bool bInForEvaluation) override;
+	virtual bool ImportPreset(const UPropertyAnimatorCorePresetBase* InPreset, const TSharedRef<FJsonValue>& InValue) override;
+	virtual bool ExportPreset(const UPropertyAnimatorCorePresetBase* InPreset, TSharedPtr<FJsonValue>& OutValue) override;
 	//~ End UPropertyAnimatorCoreResolver
+
+protected:
+	UPROPERTY(EditInstanceOnly, Category="Animator")
+	EPropertyAnimatorTextResolverRangeUnit Unit = EPropertyAnimatorTextResolverRangeUnit::Percentage;
+
+	UPROPERTY(EditInstanceOnly, Setter, Getter, Category="Animator", meta=(ClampMin="0", ClampMax="100", Units=Percent, EditCondition="Unit == EPropertyAnimatorTextResolverRangeUnit::Percentage", EditConditionHides))
+	float Start = 0.f;
+
+	UPROPERTY(EditInstanceOnly, Setter, Getter, Category="Animator", meta=(ClampMin="0", ClampMax="100", Units=Percent, EditCondition="Unit == EPropertyAnimatorTextResolverRangeUnit::Percentage", EditConditionHides))
+	float End = 100.f;
+
+	UPROPERTY(EditInstanceOnly, Setter, Getter, Category="Animator", meta=(Units=Percent, Delta="1", EditCondition="Unit == EPropertyAnimatorTextResolverRangeUnit::Percentage", EditConditionHides))
+	float Offset = 0.f;
+
+	UPROPERTY(EditInstanceOnly, Category="Animator", meta=(ClampMin="0", EditCondition="Unit == EPropertyAnimatorTextResolverRangeUnit::Character", EditConditionHides))
+	int32 CharacterStartIndex = 0;
+
+	UPROPERTY(EditInstanceOnly, Category="Animator", meta=(ClampMin="0", EditCondition="Unit == EPropertyAnimatorTextResolverRangeUnit::Character", EditConditionHides))
+	int32 CharacterEndIndex = 100;
+
+	UPROPERTY(EditInstanceOnly, Category="Animator", meta=(EditCondition="Unit == EPropertyAnimatorTextResolverRangeUnit::Character", EditConditionHides))
+	int32 CharacterOffsetIndex = 0;
+
+	UPROPERTY(EditInstanceOnly, Category="Animator", meta=(ClampMin="0", EditCondition="Unit == EPropertyAnimatorTextResolverRangeUnit::Word", EditConditionHides))
+	int32 WordStartIndex = 0;
+
+	UPROPERTY(EditInstanceOnly, Category="Animator", meta=(ClampMin="0", EditCondition="Unit == EPropertyAnimatorTextResolverRangeUnit::Word", EditConditionHides))
+	int32 WordEndIndex = 100;
+
+	UPROPERTY(EditInstanceOnly, Category="Animator", meta=(EditCondition="Unit == EPropertyAnimatorTextResolverRangeUnit::Word", EditConditionHides))
+	int32 WordOffsetIndex = 0;
+
+	UPROPERTY(EditInstanceOnly, Category="Animator")
+	EPropertyAnimatorTextResolverRangeDirection Direction = EPropertyAnimatorTextResolverRangeDirection::LeftToRight;
 };
\ No newline at end of file
diff --git a/Engine/Plugins/Experimental/PropertyAnimator/Source/PropertyAnimator/Public/Properties/PropertyAnimatorTextGroup.h b/Engine/Plugins/Experimental/PropertyAnimator/Source/PropertyAnimator/Public/Properties/PropertyAnimatorTextGroup.h
deleted file mode 100644
diff --git a/Engine/Plugins/Experimental/PropertyAnimatorCore/Source/PropertyAnimatorCore/Private/Animators/PropertyAnimatorCoreBase.cpp b/Engine/Plugins/Experimental/PropertyAnimatorCore/Source/PropertyAnimatorCore/Private/Animators/PropertyAnimatorCoreBase.cpp
--- a/Engine/Plugins/Experimental/PropertyAnimatorCore/Source/PropertyAnimatorCore/Private/Animators/PropertyAnimatorCoreBase.cpp
+++ b/Engine/Plugins/Experimental/PropertyAnimatorCore/Source/PropertyAnimatorCore/Private/Animators/PropertyAnimatorCoreBase.cpp
@@ -5,7 +5,6 @@
 #include "Components/PropertyAnimatorCoreComponent.h"
 #include "Dom/JsonObject.h"
 #include "GameFramework/Actor.h"
-#include "Properties/PropertyAnimatorCoreGroupBase.h"
 #include "Properties/PropertyAnimatorCoreResolver.h"
 #include "Properties/Handlers/PropertyAnimatorCoreHandlerBase.h"
 #include "Settings/PropertyAnimatorCoreSettings.h"
@@ -205,10 +204,6 @@ void UPropertyAnimatorCoreBase::PostEditChangeProperty(FPropertyChangedEvent& Pr
 	{
 		OnTimeSourceNameChanged();
 	}
-	else if (MemberName == GET_MEMBER_NAME_CHECKED(UPropertyAnimatorCoreBase, PropertyGroups))
-	{
-		OnPropertyGroupsChanged();
-	}
 }
 #endif
 
@@ -604,36 +599,6 @@ void UPropertyAnimatorCoreBase::OnObjectReplaced(const TMap<UObject*, UObject*>&
 	}, bResolve);
 }
 
-void UPropertyAnimatorCoreBase::OnPropertyGroupsChanged()
-{
-	TSet<FName> CurrentGroupNames;
-
-	Algo::TransformIf(
-		PropertyGroups
-		, CurrentGroupNames
-		, [](const UPropertyAnimatorCoreGroupBase* InGroup)->bool
-		{
-			return !!InGroup;
-		}
-		, [](const UPropertyAnimatorCoreGroupBase* InGroup)->FName
-		{
-			return InGroup->GetFName();
-		}
-	);
-
-	// Remove assigned group in property context if removed from animator
-	constexpr bool bResolve = false;
-	ForEachLinkedProperty<UPropertyAnimatorCoreContext>([&CurrentGroupNames](UPropertyAnimatorCoreContext* InContext, const FPropertyAnimatorCoreData& InProperty)->bool
-	{
-		if (!CurrentGroupNames.Contains(InContext->GroupName))
-		{
-			InContext->SetGroup(nullptr);
-		}
-
-		return true;
-	}, bResolve);
-}
-
 void UPropertyAnimatorCoreBase::RestoreProperties(bool bInForce)
 {
 	constexpr bool bResolve = false;
diff --git a/Engine/Plugins/Experimental/PropertyAnimatorCore/Source/PropertyAnimatorCore/Private/Properties/PropertyAnimatorCoreContext.cpp b/Engine/Plugins/Experimental/PropertyAnimatorCore/Source/PropertyAnimatorCore/Private/Properties/PropertyAnimatorCoreContext.cpp
--- a/Engine/Plugins/Experimental/PropertyAnimatorCore/Source/PropertyAnimatorCore/Private/Properties/PropertyAnimatorCoreContext.cpp
+++ b/Engine/Plugins/Experimental/PropertyAnimatorCore/Source/PropertyAnimatorCore/Private/Properties/PropertyAnimatorCoreContext.cpp
@@ -6,7 +6,6 @@
 #include "Dom/JsonObject.h"
 #include "Dom/JsonValue.h"
 #include "StructUtils/InstancedStruct.h"
-#include "Properties/PropertyAnimatorCoreGroupBase.h"
 #include "Properties/PropertyAnimatorCoreResolver.h"
 #include "Properties/Converters/PropertyAnimatorCoreConverterBase.h"
 #include "Properties/Handlers/PropertyAnimatorCoreHandlerBase.h"
@@ -18,20 +17,15 @@ TArray<FPropertyAnimatorCoreData> UPropertyAnimatorCoreContext::ResolveProperty(
 {
 	TArray<FPropertyAnimatorCoreData> ResolvedProperties;
 
-	if (UPropertyAnimatorCoreResolver* Resolver = AnimatedProperty.GetPropertyResolver())
+	if (UPropertyAnimatorCoreResolver* PropertyResolver = GetResolver())
 	{
-		Resolver->ResolveProperties(AnimatedProperty, ResolvedProperties);
+		PropertyResolver->ResolveProperties(AnimatedProperty, ResolvedProperties, bInForEvaluation);
 	}
 	else
 	{
 		ResolvedProperties.Add(AnimatedProperty);
 	}
 
-	if (bInForEvaluation && Group)
-	{
-		Group->ManageProperties(this, ResolvedProperties);
-	}
-
 	return ResolvedProperties;
 }
 
@@ -59,18 +53,14 @@ UPropertyAnimatorCoreHandlerBase* UPropertyAnimatorCoreContext::GetHandler() con
 	return HandlerWeak.Get();
 }
 
-void UPropertyAnimatorCoreContext::SetGroup(UPropertyAnimatorCoreGroupBase* InGroup)
+UPropertyAnimatorCoreResolver* UPropertyAnimatorCoreContext::GetResolver() const
 {
-	if (InGroup && InGroup->IsPropertySupported(this))
-	{
-		Group = InGroup;
-		GroupName = InGroup->GetFName();
-	}
-	else
+	if (Resolver)
 	{
-		Group = nullptr;
-		GroupName = NAME_None;
+		return Resolver;
 	}
+
+	return AnimatedProperty.GetPropertyResolver();
 }
 
 bool UPropertyAnimatorCoreContext::IsResolvable() const
@@ -130,35 +120,27 @@ void UPropertyAnimatorCoreContext::SetConverterClass(TSubclassOf<UPropertyAnimat
 	}
 }
 
-void UPropertyAnimatorCoreContext::SetGroupName(FName InGroupName)
-{
-	if (GroupName.IsEqual(InGroupName))
-	{
-		return;
-	}
-
-	const TArray<FName> GroupNames = GetSupportedGroupNames();
-	if (!GroupNames.Contains(InGroupName))
-	{
-		return;
-	}
-
-	GroupName = InGroupName;
-	OnGroupNameChanged();
-}
-
 void UPropertyAnimatorCoreContext::PostLoad()
 {
 	Super::PostLoad();
 
 	CheckEditMode();
 	CheckEditConverterRule();
+	CheckEditResolver();
 
 	FTSTicker::GetCoreTicker().AddTicker(FTickerDelegate::CreateWeakLambda(this, [this](float InDelta)
 	{
 		// Restore before regenerating new property path
 		Restore();
+
+		if (AnimatedProperty.IsResolvable())
+		{
+			DeltaPropertyValues.Reset();
+			OriginalPropertyValues.Reset();
+		}
+
 		AnimatedProperty.GeneratePropertyPath();
+
 		return false;
 	}));
 }
@@ -195,10 +177,6 @@ void UPropertyAnimatorCoreContext::PostEditChangeProperty(FPropertyChangedEvent&
 	{
 		OnAnimatedChanged();
 	}
-	else if (MemberName == GET_MEMBER_NAME_CHECKED(UPropertyAnimatorCoreContext, GroupName))
-	{
-		OnGroupNameChanged();
-	}
 }
 #endif
 
@@ -232,6 +210,14 @@ bool UPropertyAnimatorCoreContext::ImportPreset(const UPropertyAnimatorCorePrese
 		SetMode(static_cast<EPropertyAnimatorCoreMode>(JsonMode));
 	}
 
+	if (Resolver)
+	{
+		if ((*JsonObject)->HasTypedField(GET_MEMBER_NAME_STRING_CHECKED(UPropertyAnimatorCoreContext, Resolver), EJson::Object))
+		{
+			Resolver->ImportPreset(InPreset, (*JsonObject)->TryGetField(GET_MEMBER_NAME_STRING_CHECKED(UPropertyAnimatorCoreContext, Resolver)).ToSharedRef());
+		}
+	}
+
 	return true;
 }
 
@@ -255,12 +241,30 @@ bool UPropertyAnimatorCoreContext::ExportPreset(const UPropertyAnimatorCorePrese
 
 	JsonObject->SetStringField(GET_MEMBER_NAME_STRING_CHECKED(UPropertyAnimatorCoreContext, AnimatedProperty), AnimatedProperty.GetPropertyLocatorPath());
 
+	if (Resolver)
+	{
+		TSharedPtr<FJsonValue> JsonValue;
+		if (Resolver->ExportPreset(InPreset, JsonValue) && JsonValue.IsValid())
+		{
+			JsonObject->SetField(GET_MEMBER_NAME_STRING_CHECKED(UPropertyAnimatorCoreContext, Resolver), JsonValue);
+		}
+	}
+
 	return true;
 }
 
 void UPropertyAnimatorCoreContext::OnAnimatedPropertyLinked()
 {
 	bEditMagnitude = AnimatedProperty.IsA<FNumericProperty>() || AnimatedProperty.HasA<FNumericProperty>();
+
+	if (const UPropertyAnimatorCoreResolver* PropertyResolver = AnimatedProperty.GetPropertyResolver())
+	{
+		if (!PropertyResolver->GetClass()->HasAnyClassFlags(CLASS_Abstract | CLASS_Transient))
+		{
+			bEditResolver = true;
+			Resolver = NewObject<UPropertyAnimatorCoreResolver>(this, PropertyResolver->GetClass());
+		}
+	}
 }
 
 void UPropertyAnimatorCoreContext::OnModeChanged()
@@ -276,19 +280,6 @@ void UPropertyAnimatorCoreContext::OnModeChanged()
 	}
 }
 
-void UPropertyAnimatorCoreContext::OnGroupNameChanged()
-{
-	if (const UPropertyAnimatorCoreBase* Animator = GetAnimator())
-	{
-		const TObjectPtr<UPropertyAnimatorCoreGroupBase>* PropertyGroup = Animator->PropertyGroups.FindByPredicate([this](const UPropertyAnimatorCoreGroupBase* InGroup)
-		{
-			return InGroup && InGroup->GetFName() == GroupName;
-		});
-
-		SetGroup(PropertyGroup ? PropertyGroup->Get() : nullptr);
-	}
-}
-
 bool UPropertyAnimatorCoreContext::ResolvePropertyOwner(AActor* InNewOwner)
 {
 	UObject* NewOwner = InNewOwner ? InNewOwner : GetTypedOuter<AActor>();
@@ -396,6 +387,7 @@ void UPropertyAnimatorCoreContext::ConstructInternal(const FPropertyAnimatorCore
 	AnimatedProperty = InProperty;
 	CheckEditMode();
 	CheckEditConverterRule();
+	CheckEditResolver();
 	SetMode(EPropertyAnimatorCoreMode::Additive);
 	OnAnimatedPropertyLinked();
 }
@@ -432,6 +424,11 @@ void UPropertyAnimatorCoreContext::CheckEditConverterRule()
 	bEditConverterRule = ConverterRule.IsValid();
 }
 
+void UPropertyAnimatorCoreContext::CheckEditResolver()
+{
+	bEditResolver = AnimatedProperty.IsResolvable();
+}
+
 void* UPropertyAnimatorCoreContext::GetConverterRulePtr(const UScriptStruct* InStruct)
 {
 	if (ConverterRule.IsValid() && ConverterRule.GetScriptStruct()->IsChildOf(InStruct))
@@ -514,27 +511,6 @@ void UPropertyAnimatorCoreContext::Save()
 	}
 }
 
-TArray<FName> UPropertyAnimatorCoreContext::GetSupportedGroupNames() const
-{
-	TArray<FName> GroupNames
-	{
-		NAME_None
-	};
-
-	if (UPropertyAnimatorCoreBase* Animator = GetAnimator())
-	{
-		for (const UPropertyAnimatorCoreGroupBase* PropertyGroup : Animator->PropertyGroups)
-		{
-			if (PropertyGroup && PropertyGroup->IsPropertySupported(this))
-			{
-				GroupNames.Add(PropertyGroup->GetFName());
-			}
-		}
-	}
-
-	return GroupNames;
-}
-
 void UPropertyAnimatorCoreContext::OnAnimatedChanged()
 {
 	if (!bAnimated)
diff --git a/Engine/Plugins/Experimental/PropertyAnimatorCore/Source/PropertyAnimatorCore/Private/Properties/PropertyAnimatorCoreResolver.cpp b/Engine/Plugins/Experimental/PropertyAnimatorCore/Source/PropertyAnimatorCore/Private/Properties/PropertyAnimatorCoreResolver.cpp
new file mode 100644
--- /dev/null
+++ b/Engine/Plugins/Experimental/PropertyAnimatorCore/Source/PropertyAnimatorCore/Private/Properties/PropertyAnimatorCoreResolver.cpp
@@ -0,0 +1,25 @@
+// Copyright Epic Games, Inc. All Rights Reserved.
+
+#include "Properties/PropertyAnimatorCoreResolver.h"
+
+#include "Dom/JsonObject.h"
+#include "Dom/JsonValue.h"
+
+bool UPropertyAnimatorCoreResolver::ImportPreset(const UPropertyAnimatorCorePresetBase* InPreset, const TSharedRef<FJsonValue>& InValue)
+{
+	const TSharedPtr<FJsonObject>* JsonTimeSourceObject;
+	if (!InValue->TryGetObject(JsonTimeSourceObject) || !JsonTimeSourceObject)
+	{
+		return false;
+	}
+
+	return true;
+}
+
+bool UPropertyAnimatorCoreResolver::ExportPreset(const UPropertyAnimatorCorePresetBase* InPreset, TSharedPtr<FJsonValue>& OutValue)
+{
+	TSharedRef<FJsonObject> JsonTimeSourceObject = MakeShared<FJsonObject>();
+	OutValue = MakeShared<FJsonValueObject>(JsonTimeSourceObject);
+
+	return true;
+}
diff --git a/Engine/Plugins/Experimental/PropertyAnimatorCore/Source/PropertyAnimatorCore/Public/Animators/PropertyAnimatorCoreBase.h b/Engine/Plugins/Experimental/PropertyAnimatorCore/Source/PropertyAnimatorCore/Public/Animators/PropertyAnimatorCoreBase.h
--- a/Engine/Plugins/Experimental/PropertyAnimatorCore/Source/PropertyAnimatorCore/Public/Animators/PropertyAnimatorCoreBase.h
+++ b/Engine/Plugins/Experimental/PropertyAnimatorCore/Source/PropertyAnimatorCore/Public/Animators/PropertyAnimatorCoreBase.h
@@ -9,7 +9,6 @@
 #include "PropertyAnimatorCoreBase.generated.h"
 
 class UPropertyAnimatorCoreComponent;
-class UPropertyAnimatorCoreGroupBase;
 class UPropertyAnimatorCoreTimeSourceBase;
 
 UENUM(BlueprintType)
@@ -357,10 +356,6 @@ private:
 	UPROPERTY(EditInstanceOnly, NoClear, Export, Instanced, EditFixedSize, Category="Animator", meta=(EditFixedOrder))
 	TArray<TObjectPtr<UPropertyAnimatorCoreContext>> LinkedProperties;
 
-	/** Groups for properties linked to this Animator */
-	UPROPERTY()
-	TArray<TObjectPtr<UPropertyAnimatorCoreGroupBase>> PropertyGroups;
-
 	/** Use the global time source or override it on this animator */
 	UPROPERTY(EditInstanceOnly, Setter="SetOverrideTimeSource", Getter="GetOverrideTimeSource", Category="Animator", meta=(InlineEditConditionToggle))
 	bool bOverrideTimeSource = true;
diff --git a/Engine/Plugins/Experimental/PropertyAnimatorCore/Source/PropertyAnimatorCore/Public/Properties/PropertyAnimatorCoreContext.h b/Engine/Plugins/Experimental/PropertyAnimatorCore/Source/PropertyAnimatorCore/Public/Properties/PropertyAnimatorCoreContext.h
--- a/Engine/Plugins/Experimental/PropertyAnimatorCore/Source/PropertyAnimatorCore/Public/Properties/PropertyAnimatorCoreContext.h
+++ b/Engine/Plugins/Experimental/PropertyAnimatorCore/Source/PropertyAnimatorCore/Public/Properties/PropertyAnimatorCoreContext.h
@@ -12,7 +12,7 @@
 struct FInstancedStruct;
 class UPropertyAnimatorCoreBase;
 class UPropertyAnimatorCoreConverterBase;
-class UPropertyAnimatorCoreGroupBase;
+class UPropertyAnimatorCoreResolver;
 class UScriptStruct;
 
 /** Mode supported for properties value */
@@ -47,11 +47,8 @@ public:
 	/** Get the handler responsible for this property type */
 	UPropertyAnimatorCoreHandlerBase* GetHandler() const;
 
-	/** Get the active group of this property */
-	UPropertyAnimatorCoreGroupBase* GetGroup() const
-	{
-		return Group;
-	}
+	/** Get the active resolver for this property if any */
+	UPropertyAnimatorCoreResolver* GetResolver() const;
 
 	PROPERTYANIMATORCORE_API void SetAnimated(bool bInAnimated);
 	bool IsAnimated() const
@@ -83,12 +80,6 @@ public:
 		return ConverterClass;
 	}
 
-	PROPERTYANIMATORCORE_API void SetGroupName(FName InGroupName);
-	FName GetGroupName() const
-	{
-		return GroupName;
-	}
-
 	/** Get converter rule if any */
 	template <typename InRuleType
 		UE_REQUIRES(TModels_V<CStaticStructProvider, InRuleType>)>
@@ -134,10 +125,10 @@ private:
 
 	void CheckEditMode();
 	void CheckEditConverterRule();
+	void CheckEditResolver();
 
 	void OnAnimatedChanged();
 	void OnModeChanged();
-	void OnGroupNameChanged();
 
 	/** Sets the evaluation result for the resolved property */
 	PROPERTYANIMATORCORE_API void CommitEvaluationResult(const FPropertyAnimatorCoreData& InResolvedProperty, const FInstancedPropertyBag& InEvaluatedValues);
@@ -151,15 +142,9 @@ private:
 	/** Allocate and save properties */
 	void Save();
 
-	void SetGroup(UPropertyAnimatorCoreGroupBase* InGroup);
-
 	bool IsResolvable() const;
 	bool IsConverted() const;
 
-	/** Get the supported group names that can manage this property */
-	UFUNCTION()
-	TArray<FName> GetSupportedGroupNames() const;
-
 	/** Animation is enabled for this property */
 	UPROPERTY(EditInstanceOnly, Setter="SetAnimated", Getter="IsAnimated", Category="Animator", meta=(AllowPrivateAccess="true"))
 	bool bAnimated = true;
@@ -192,13 +177,13 @@ private:
 	UPROPERTY(EditInstanceOnly, Category="Animator", meta=(HideEditConditionToggle, EditCondition="bEditConverterRule", EditConditionHides, AllowPrivateAccess="true"))
 	FInstancedStruct ConverterRule;
 
-	/** The unique group name that manages this property */
-	UPROPERTY()
-	FName GroupName = NAME_None;
+	/** Edit condition for property resolver */
+	UPROPERTY(Transient)
+	bool bEditResolver = true;
 
-	/** Active group of this property */
-	UPROPERTY()
-	TObjectPtr<UPropertyAnimatorCoreGroupBase> Group;
+	/** Custom resolver for the property */
+	UPROPERTY(VisibleInstanceOnly, NoClear, Export, Instanced, DisplayName="Range", Category="Animator", meta=(EditCondition="bEditResolver", EditConditionHides, HideEditConditionToggle, AllowPrivateAccess="true"))
+	TObjectPtr<UPropertyAnimatorCoreResolver> Resolver;
 
 	/** Store original property values for resolved properties */
 	UPROPERTY(NonTransactional)
diff --git a/Engine/Plugins/Experimental/PropertyAnimatorCore/Source/PropertyAnimatorCore/Public/Properties/PropertyAnimatorCoreGroupBase.h b/Engine/Plugins/Experimental/PropertyAnimatorCore/Source/PropertyAnimatorCore/Public/Properties/PropertyAnimatorCoreGroupBase.h
deleted file mode 100644
diff --git a/Engine/Plugins/Experimental/PropertyAnimatorCore/Source/PropertyAnimatorCore/Public/Properties/PropertyAnimatorCoreResolver.h b/Engine/Plugins/Experimental/PropertyAnimatorCore/Source/PropertyAnimatorCore/Public/Properties/PropertyAnimatorCoreResolver.h
--- a/Engine/Plugins/Experimental/PropertyAnimatorCore/Source/PropertyAnimatorCore/Public/Properties/PropertyAnimatorCoreResolver.h
+++ b/Engine/Plugins/Experimental/PropertyAnimatorCore/Source/PropertyAnimatorCore/Public/Properties/PropertyAnimatorCoreResolver.h
@@ -2,6 +2,7 @@
 
 #pragma once
 
+#include "Presets/PropertyAnimatorCorePresetable.h"
 #include "Properties/PropertyAnimatorCoreData.h"
 #include "UObject/NameTypes.h"
 #include "UObject/Object.h"
@@ -11,44 +12,38 @@
  * Base class to find properties hidden or not reachable,
  * allows to discover resolvable properties for specific actors/components/objects
  * that we cannot reach or are transient, will be resolved when needed
- * Resolvers should remain transient and stateless
  */
-UCLASS(MinimalAPI, Abstract, Transient)
-class UPropertyAnimatorCoreResolver : public UObject
+UCLASS(MinimalAPI, Abstract)
+class UPropertyAnimatorCoreResolver : public UObject, public IPropertyAnimatorCorePresetable
 {
 	GENERATED_BODY()
 
 public:
 	UPropertyAnimatorCoreResolver()
-		: UPropertyAnimatorCoreResolver(NAME_None, false)
+		: UPropertyAnimatorCoreResolver(NAME_None)
 	{}
 
-	UPropertyAnimatorCoreResolver(FName InResolverName, bool bInIsRange)
+	UPropertyAnimatorCoreResolver(FName InResolverName)
 		: ResolverName(InResolverName)
-		, bIsRange(bInIsRange)
 	{}
 
 	/** Get properties found inside parent property */
 	virtual void GetResolvableProperties(const FPropertyAnimatorCoreData& InParentProperty, TSet<FPropertyAnimatorCoreData>& OutProperties) {}
 
 	/** Called when we actually need the underlying properties */
-	virtual void ResolveProperties(const FPropertyAnimatorCoreData& InTemplateProperty, TArray<FPropertyAnimatorCoreData>& OutProperties) {}
+	virtual void ResolveProperties(const FPropertyAnimatorCoreData& InTemplateProperty, TArray<FPropertyAnimatorCoreData>& OutProperties, bool bInForEvaluation) {}
 
 	FName GetResolverName() const
 	{
 		return ResolverName;
 	}
 
-	/** Whether this property when resolved should be treated as a range */
-	bool IsRange() const
-	{
-		return bIsRange;
-	}
+	//~ Begin IPropertyAnimatorCorePresetable
+	PROPERTYANIMATORCORE_API virtual bool ImportPreset(const UPropertyAnimatorCorePresetBase* InPreset, const TSharedRef<FJsonValue>& InValue) override;
+	PROPERTYANIMATORCORE_API virtual bool ExportPreset(const UPropertyAnimatorCorePresetBase* InPreset, TSharedPtr<FJsonValue>& OutValue) override;
+	//~ End IPropertyAnimatorCorePresetable
 
 private:
 	UPROPERTY()
 	FName ResolverName = NAME_None;
-
-	UPROPERTY()
-	bool bIsRange = false;
-};
\ No newline at end of file
+};
diff --git a/Engine/Plugins/Experimental/Text3D/Source/Text3D/Private/Text3DComponent.cpp b/Engine/Plugins/Experimental/Text3D/Source/Text3D/Private/Text3DComponent.cpp
--- a/Engine/Plugins/Experimental/Text3D/Source/Text3D/Private/Text3DComponent.cpp
+++ b/Engine/Plugins/Experimental/Text3D/Source/Text3D/Private/Text3DComponent.cpp
@@ -12,6 +12,7 @@
 #include "Fonts/SlateTextShaper.h"
 #include "Framework/Text/PlainTextLayoutMarshaller.h"
 #include "Glyph.h"
+#include "Internationalization/Regex.h"
 #include "Materials/Material.h"
 #include "MeshCreator.h"
 #include "Misc/ScopeExit.h"
@@ -279,6 +280,39 @@ void UText3DComponent::RefreshTypeface()
 	}
 }
 
+void UText3DComponent::UpdateStatistics()
+{
+	Statistics = FText3DStatistics();
+
+	const FString WordString = Text.ToString();
+
+	const FRegexPattern WordPattern(TEXT("\\S+"));
+	FRegexMatcher Matcher(WordPattern, WordString);
+
+	int32 PreviousEndIndex = 0;
+	int32 WhitespaceCount = 0;
+
+	while (Matcher.FindNext())
+	{
+		const FString Word = Matcher.GetCaptureGroup(0);
+
+		if (!Word.IsEmpty())
+		{
+			FText3DWordStatistics& WordStatistics = Statistics.Words.Add_GetRef(FText3DWordStatistics());
+			const int32 MatchBegin = Matcher.GetMatchBeginning();
+			const int32 MatchEnd = Matcher.GetMatchEnding();
+
+			WordStatistics.ActualRange = FTextRange(MatchBegin, MatchEnd);
+
+			WhitespaceCount += MatchBegin - PreviousEndIndex;
+
+			WordStatistics.RenderRange = FTextRange(MatchBegin - WhitespaceCount, MatchEnd - WhitespaceCount);
+
+			PreviousEndIndex = MatchEnd;
+		}
+	}
+}
+
 void UText3DComponent::OnRegister()
 {
 	Super::OnRegister();
@@ -1377,6 +1411,7 @@ void UText3DComponent::BuildTextMeshInternal(const bool& bCleanCache)
 	}
 
 	OnMaterialChanged();
+	UpdateStatistics();
 
 	TextGeneratedNativeDelegate.Broadcast();
 	TextGeneratedDelegate.Broadcast();
diff --git a/Engine/Plugins/Experimental/Text3D/Source/Text3D/Public/Text3DComponent.h b/Engine/Plugins/Experimental/Text3D/Source/Text3D/Public/Text3DComponent.h
--- a/Engine/Plugins/Experimental/Text3D/Source/Text3D/Public/Text3DComponent.h
+++ b/Engine/Plugins/Experimental/Text3D/Source/Text3D/Public/Text3DComponent.h
@@ -53,6 +53,21 @@ enum class EText3DModifyFlags : uint8
 };
 ENUM_CLASS_FLAGS(EText3DModifyFlags)
 
+struct FText3DWordStatistics
+{
+	/** Actual range taking into account whitespaces */
+	FTextRange ActualRange;
+
+	/** Render range not taking into account whitespaces */
+	FTextRange RenderRange;
+};
+
+struct FText3DStatistics
+{
+	TArray<FText3DWordStatistics> Words;
+	int32 WhiteSpaces;
+};
+
 UCLASS(ClassGroup = (Text3D), PrioritizeCategories = "Text Layout Geometry Materials", meta = (BlueprintSpawnableComponent))
 class TEXT3D_API UText3DComponent : public USceneComponent
 {
@@ -299,6 +314,11 @@ public:
 	/** Manually update the geometry, ignoring RefreshOnChange (but still accounting for the Freeze flag) */
 	void Rebuild();
 
+	const FText3DStatistics& GetStatistics() const
+	{
+		return Statistics;
+	}
+
 protected:
 	/** Whether to allow automatic refresh/mesh generation */
 	UPROPERTY(EditAnywhere, BlueprintReadWrite, Getter = "RefreshesOnChange", Setter = "SetRefreshOnChange", Category = "Text3D", AdvancedDisplay, meta = (AllowPrivateAccess = "true"))
@@ -475,6 +495,9 @@ private:
 	UPROPERTY(Transient)
 	TArray<TObjectPtr<UStaticMeshComponent>> CharacterMeshes;
 
+	/** Text statistics cached since last text generation */
+	FText3DStatistics Statistics;
+
 	/** Allocates, or shrinks existing components to match the input number. Returns false if nothing modified. */
 	bool AllocateGlyphs(int32 Num);
 
@@ -509,4 +532,6 @@ private:
 	bool IsTypefaceAvailable(FName InTypeface) const;
 	TArray<FTypefaceEntry> GetAvailableTypefaces() const;
 	void RefreshTypeface();
+
+	void UpdateStatistics();
 };
\ No newline at end of file
