MotionDesign : Animators
- Force show engine content and plugin content to see any curves for the curve animator

diff --git a/Engine/Plugins/Experimental/PropertyAnimator/Source/PropertyAnimator/Public/Animators/PropertyAnimatorCurve.h b/Engine/Plugins/Experimental/PropertyAnimator/Source/PropertyAnimator/Public/Animators/PropertyAnimatorCurve.h
--- a/Engine/Plugins/Experimental/PropertyAnimator/Source/PropertyAnimator/Public/Animators/PropertyAnimatorCurve.h
+++ b/Engine/Plugins/Experimental/PropertyAnimator/Source/PropertyAnimator/Public/Animators/PropertyAnimatorCurve.h
@@ -13,7 +13,7 @@ struct FPropertyAnimatorCurveEasing
 {
 	GENERATED_BODY()
 
-	UPROPERTY(EditInstanceOnly, DisplayName="Curve", Category="Animator")
+	UPROPERTY(EditInstanceOnly, DisplayName="Curve", Category="Animator", meta=(ForceShowEngineContent, ForceShowPluginContent))
 	TObjectPtr<UPropertyAnimatorEaseCurve> EaseCurve;
 
 	UPROPERTY(EditInstanceOnly, Interp, DisplayName="Duration", Category="Animator", meta=(ClampMin="0", Units=Seconds))
@@ -93,7 +93,7 @@ protected:
 	FPropertyAnimatorCurveEasing EaseIn;
 
 	/** The base curve to sample for the animation */
-	UPROPERTY(EditInstanceOnly, Setter, Getter, DisplayName="Loop Curve", Category="Animator")
+	UPROPERTY(EditInstanceOnly, Setter, Getter, DisplayName="Loop Curve", Category="Animator", meta=(ForceShowEngineContent, ForceShowPluginContent))
 	TObjectPtr<UPropertyAnimatorWaveCurve> WaveCurve;
 
 	/** Use ease out effect */
