MotionDesign : PropertyAnimator
- Fix behavior with TimeOffset option and non-range properties where offset was not taken into account

diff --git a/Engine/Plugins/Experimental/PropertyAnimator/Source/PropertyAnimator/Private/Animators/PropertyAnimatorNumericBase.cpp b/Engine/Plugins/Experimental/PropertyAnimator/Source/PropertyAnimator/Private/Animators/PropertyAnimatorNumericBase.cpp
--- a/Engine/Plugins/Experimental/PropertyAnimator/Source/PropertyAnimator/Private/Animators/PropertyAnimatorNumericBase.cpp
+++ b/Engine/Plugins/Experimental/PropertyAnimator/Source/PropertyAnimator/Private/Animators/PropertyAnimatorNumericBase.cpp
@@ -166,35 +166,31 @@ void UPropertyAnimatorNumericBase::EvaluateProperties(FInstancedPropertyBag& InP
 	{
 		const double RandomTimeOffset = bRandomTimeOffset ? RandomStream.GetFraction() : 0;
 
-		const double TimeOffset = InOptions->GetTimeOffset() / (InRangeMax + 1);
+		const int32 RangeIndex = InRangeIndex + 1;
+		const int32 RangeMax = InRangeMax + 1;
+		const double TimeOffset = InOptions->GetTimeOffset() / RangeMax;
 		const double AbsTimeOffset = FMath::Abs(TimeOffset);
-		const double MaxTimeOffset = InRangeMax * AbsTimeOffset;
-		double PropertyTimeElapsed = TimeElapsed - MaxTimeOffset + RandomTimeOffset;
+		const double MaxTimeOffset = RangeMax * AbsTimeOffset;
+		double PropertyTimeElapsed = TimeElapsed + RandomTimeOffset;
 		double Frequency = CycleDuration != 0.f ? 1.f / CycleDuration : 0.f;
 
-		if (TimeOffset >= 0)
-		{
-			PropertyTimeElapsed += InRangeIndex * AbsTimeOffset;
-		}
-		else
-		{
-			PropertyTimeElapsed += MaxTimeOffset - InRangeIndex * AbsTimeOffset;
-		}
+		PropertyTimeElapsed += RangeIndex * TimeOffset;
+		PropertyTimeElapsed = FMath::Abs(PropertyTimeElapsed);
 
 		if (CycleMode == EPropertyAnimatorCycleMode::DoOnce)
 		{
-			if (FMath::Abs(PropertyTimeElapsed) > CycleDuration)
+			if (PropertyTimeElapsed > CycleDuration)
 			{
-				PropertyTimeElapsed = CycleDuration;
+				PropertyTimeElapsed = CycleDuration - UE_KINDA_SMALL_NUMBER;
 			}
 		}
 		else if (CycleMode == EPropertyAnimatorCycleMode::Loop)
 		{
 			PropertyTimeElapsed = FMath::Fmod(PropertyTimeElapsed, CycleDuration + MaxTimeOffset + CycleGapDuration);
 
-			if (FMath::Abs(PropertyTimeElapsed) > CycleDuration)
+			if (PropertyTimeElapsed > CycleDuration)
 			{
-				PropertyTimeElapsed = CycleDuration;
+				PropertyTimeElapsed = CycleDuration - UE_KINDA_SMALL_NUMBER;
 			}
 		}
 		else if (CycleMode == EPropertyAnimatorCycleMode::PingPong)
@@ -202,9 +198,9 @@ void UPropertyAnimatorNumericBase::EvaluateProperties(FInstancedPropertyBag& InP
 			const bool bReverse = FMath::Modulo(FMath::TruncToInt32(PropertyTimeElapsed / (CycleDuration + MaxTimeOffset + CycleGapDuration)), 2) != 0;
 			PropertyTimeElapsed = FMath::Fmod(PropertyTimeElapsed, CycleDuration + MaxTimeOffset + CycleGapDuration);
 
-			if (FMath::Abs(PropertyTimeElapsed) > CycleDuration)
+			if (PropertyTimeElapsed > CycleDuration)
 			{
-				PropertyTimeElapsed = CycleDuration;
+				PropertyTimeElapsed = CycleDuration - UE_KINDA_SMALL_NUMBER;
 			}
 
 			if (bReverse)
