MotionDesign : PropertyAnimatorCore
- Removing CycleRate property from numeric animator, moving it to wiggle animator to avoid other animators to have it as well and not using it

diff --git a/Engine/Plugins/Experimental/PropertyAnimator/Source/PropertyAnimator/Private/Animators/PropertyAnimatorNumericBase.cpp b/Engine/Plugins/Experimental/PropertyAnimator/Source/PropertyAnimator/Private/Animators/PropertyAnimatorNumericBase.cpp
--- a/Engine/Plugins/Experimental/PropertyAnimator/Source/PropertyAnimator/Private/Animators/PropertyAnimatorNumericBase.cpp
+++ b/Engine/Plugins/Experimental/PropertyAnimator/Source/PropertyAnimator/Private/Animators/PropertyAnimatorNumericBase.cpp
@@ -56,11 +56,6 @@ void UPropertyAnimatorNumericBase::SetCycleGapDuration(float InCycleGap)
 	CycleGapDuration = FMath::Max(0, InCycleGap);
 }
 
-void UPropertyAnimatorNumericBase::SetCycleRate(float InFrequency)
-{
-	CycleRate = FMath::Max(0, InFrequency);
-}
-
 void UPropertyAnimatorNumericBase::SetCycleMode(EPropertyAnimatorCycleMode InMode)
 {
 	if (CycleMode == InMode)
@@ -214,7 +209,7 @@ void UPropertyAnimatorNumericBase::EvaluateProperties(FInstancedPropertyBag& InP
 		}
 		else if (CycleMode == EPropertyAnimatorCycleMode::None)
 		{
-			Frequency = CycleRate;
+			Frequency = 1;
 		}
 
 		if (Magnitude != 0
diff --git a/Engine/Plugins/Experimental/PropertyAnimator/Source/PropertyAnimator/Private/Animators/PropertyAnimatorWiggle.cpp b/Engine/Plugins/Experimental/PropertyAnimator/Source/PropertyAnimator/Private/Animators/PropertyAnimatorWiggle.cpp
--- a/Engine/Plugins/Experimental/PropertyAnimator/Source/PropertyAnimator/Private/Animators/PropertyAnimatorWiggle.cpp
+++ b/Engine/Plugins/Experimental/PropertyAnimator/Source/PropertyAnimator/Private/Animators/PropertyAnimatorWiggle.cpp
@@ -15,6 +15,11 @@ UPropertyAnimatorWiggle::UPropertyAnimatorWiggle()
 	CycleMode = EPropertyAnimatorCycleMode::None;
 }
 
+void UPropertyAnimatorWiggle::SetFrequency(float InFrequency)
+{
+	Frequency = FMath::Max(0, InFrequency);
+}
+
 void UPropertyAnimatorWiggle::OnAnimatorRegistered(FPropertyAnimatorCoreMetadata& InMetadata)
 {
 	Super::OnAnimatorRegistered(InMetadata);
@@ -25,7 +30,6 @@ void UPropertyAnimatorWiggle::OnAnimatorRegistered(FPropertyAnimatorCoreMetadata
 bool UPropertyAnimatorWiggle::EvaluateProperty(const FPropertyAnimatorCoreData& InPropertyData, UPropertyAnimatorCoreContext* InContext, FInstancedPropertyBag& InParameters, FInstancedPropertyBag& OutEvaluationResult) const
 {
 	const double TimeElapsed = InParameters.GetValueDouble(TimeElapsedParameterName).GetValue();
-	const double Frequency = InParameters.GetValueDouble(FrequencyParameterName).GetValue();
 
 	// Apply random wave based on time and frequency
 	const double WaveResult = UE::PropertyAnimator::Wave::Perlin(TimeElapsed, 1.f, Frequency, 0.f);
@@ -38,3 +42,33 @@ bool UPropertyAnimatorWiggle::EvaluateProperty(const FPropertyAnimatorCoreData&
 
 	return InContext->EvaluateProperty(InPropertyData, InParameters, OutEvaluationResult);
 }
+
+bool UPropertyAnimatorWiggle::ImportPreset(const UPropertyAnimatorCorePresetBase* InPreset, const TSharedRef<FPropertyAnimatorCorePresetArchive>& InValue)
+{
+	if (Super::ImportPreset(InPreset, InValue) && InValue->IsObject())
+	{
+		const TSharedPtr<FPropertyAnimatorCorePresetObjectArchive> AnimatorArchive = InValue->AsMutableObject();
+
+		double FrequencyValue = Frequency;
+		AnimatorArchive->Get(GET_MEMBER_NAME_STRING_CHECKED(UPropertyAnimatorWiggle, Frequency), FrequencyValue);
+		SetFrequency(FrequencyValue);
+
+		return true;
+	}
+
+	return false;
+}
+
+bool UPropertyAnimatorWiggle::ExportPreset(const UPropertyAnimatorCorePresetBase* InPreset, TSharedPtr<FPropertyAnimatorCorePresetArchive>& OutValue) const
+{
+	if (Super::ExportPreset(InPreset, OutValue) && OutValue->IsObject())
+	{
+		const TSharedPtr<FPropertyAnimatorCorePresetObjectArchive> AnimatorArchive = OutValue->AsMutableObject();
+
+		AnimatorArchive->Set(GET_MEMBER_NAME_STRING_CHECKED(UPropertyAnimatorWiggle, Frequency), Frequency);
+
+		return true;
+	}
+
+	return false;
+}
diff --git a/Engine/Plugins/Experimental/PropertyAnimator/Source/PropertyAnimator/Public/Animators/PropertyAnimatorNumericBase.h b/Engine/Plugins/Experimental/PropertyAnimator/Source/PropertyAnimator/Public/Animators/PropertyAnimatorNumericBase.h
--- a/Engine/Plugins/Experimental/PropertyAnimator/Source/PropertyAnimator/Public/Animators/PropertyAnimatorNumericBase.h
+++ b/Engine/Plugins/Experimental/PropertyAnimator/Source/PropertyAnimator/Public/Animators/PropertyAnimatorNumericBase.h
@@ -53,12 +53,6 @@ public:
 		return CycleGapDuration;
 	}
 
-	PROPERTYANIMATOR_API void SetCycleRate(float InFrequency);
-	float GetCycleRate() const
-	{
-		return CycleRate;
-	}
-
 	PROPERTYANIMATOR_API void SetRandomTimeOffset(bool bInOffset);
 	bool GetRandomTimeOffset() const
 	{
@@ -116,10 +110,6 @@ protected:
 	UPROPERTY(EditInstanceOnly, Setter, Getter, Category="Animator", meta=(ClampMin="0", Units=Seconds, EditCondition="CycleMode != EPropertyAnimatorCycleMode::DoOnce && CycleMode != EPropertyAnimatorCycleMode::None", EditConditionHides))
 	float CycleGapDuration = 0.f;
 
-	/** Frequency for the effect, higher values will give you faster movements */
-	UPROPERTY(EditInstanceOnly, Setter, Getter, Category="Animator", DisplayName="Frequency", meta=(ClampMin="0", Units=Hz, EditCondition="CycleMode == EPropertyAnimatorCycleMode::None", EditConditionHides))
-	float CycleRate = 1.f;
-
 	/** Use random time offset to add variation in animation */
 	UPROPERTY(EditInstanceOnly, Setter="SetRandomTimeOffset", Getter="GetRandomTimeOffset", Category="Animator", meta=(InlineEditConditionToggle))
 	bool bRandomTimeOffset = false;
diff --git a/Engine/Plugins/Experimental/PropertyAnimator/Source/PropertyAnimator/Public/Animators/PropertyAnimatorWiggle.h b/Engine/Plugins/Experimental/PropertyAnimator/Source/PropertyAnimator/Public/Animators/PropertyAnimatorWiggle.h
--- a/Engine/Plugins/Experimental/PropertyAnimator/Source/PropertyAnimator/Public/Animators/PropertyAnimatorWiggle.h
+++ b/Engine/Plugins/Experimental/PropertyAnimator/Source/PropertyAnimator/Public/Animators/PropertyAnimatorWiggle.h
@@ -16,9 +16,21 @@ class UPropertyAnimatorWiggle : public UPropertyAnimatorNumericBase
 public:
 	UPropertyAnimatorWiggle();
 
+	PROPERTYANIMATOR_API void SetFrequency(float InFrequency);
+	float GetFrequency() const
+	{
+		return Frequency;
+	}
+
 protected:
 	//~ Begin UPropertyAnimatorFloatBase
 	virtual void OnAnimatorRegistered(FPropertyAnimatorCoreMetadata& InMetadata) override;
 	virtual bool EvaluateProperty(const FPropertyAnimatorCoreData& InPropertyData, UPropertyAnimatorCoreContext* InContext, FInstancedPropertyBag& InParameters, FInstancedPropertyBag& OutEvaluationResult) const override;
+	virtual bool ImportPreset(const UPropertyAnimatorCorePresetBase* InPreset, const TSharedRef<FPropertyAnimatorCorePresetArchive>& InValue) override;
+	virtual bool ExportPreset(const UPropertyAnimatorCorePresetBase* InPreset, TSharedPtr<FPropertyAnimatorCorePresetArchive>& OutValue) const override;
 	//~ End UPropertyAnimatorFloatBase
+
+	/** Frequency for the effect, higher values will give you faster movements */
+	UPROPERTY(EditInstanceOnly, Setter, Getter, Category="Animator", meta=(ClampMin="0", Units=Hz))
+	float Frequency = 1.f;
 };
\ No newline at end of file
