MotionDesign : Text3D
- Integrating Text3D changes with property animator, instead of controlling character component, the character data will be controlled to make it agnostic of the Text3D implementation
- Migrating old text component animated property to new character property
- Restoring properties values before level is saved and disabling animation during save
- Fix Motion design template level text should not longer break due to trimmed space

diff --git a/Engine/Plugins/Experimental/PropertyAnimator/Source/PropertyAnimator/Private/Presets/PropertyAnimatorPresetTextLocation.cpp b/Engine/Plugins/Experimental/PropertyAnimator/Source/PropertyAnimator/Private/Presets/PropertyAnimatorPresetTextLocation.cpp
--- a/Engine/Plugins/Experimental/PropertyAnimator/Source/PropertyAnimator/Private/Presets/PropertyAnimatorPresetTextLocation.cpp
+++ b/Engine/Plugins/Experimental/PropertyAnimator/Source/PropertyAnimator/Private/Presets/PropertyAnimatorPresetTextLocation.cpp
@@ -3,27 +3,22 @@
 #include "Presets/PropertyAnimatorPresetTextLocation.h"
 
 #include "Animators/PropertyAnimatorCoreBase.h"
+#include "Characters/Text3DCharacterBase.h"
 #include "Properties/PropertyAnimatorTextResolver.h"
 #include "Properties/PropertyAnimatorFloatContext.h"
 #include "Text3DComponent.h"
 
 void UPropertyAnimatorPresetTextLocation::GetPresetProperties(const AActor* InActor, const UPropertyAnimatorCoreBase* InAnimator, TSet<FPropertyAnimatorCoreData>& OutProperties) const
 {
-	const UText3DComponent* TextComponent = InActor->FindComponentByClass<UText3DComponent>();
+	UText3DComponent* TextComponent = InActor->FindComponentByClass<UText3DComponent>();
 	if (!TextComponent)
 	{
 		return;
 	}
 
-	USceneComponent* TextRootComponent = TextComponent->GetChildComponent(0);
-	if (!TextRootComponent)
-	{
-		return;
-	}
-
-	const FName LocationPropertyName = USceneComponent::GetRelativeLocationPropertyName();
-	FProperty* LocationProperty = FindFProperty<FProperty>(USceneComponent::StaticClass(), LocationPropertyName);
+	const FName LocationPropertyName = UText3DCharacterBase::GetRelativeLocationPropertyName();
+	FProperty* LocationProperty = FindFProperty<FProperty>(UText3DCharacterBase::StaticClass(), LocationPropertyName);
 	check(LocationProperty);
 
-	OutProperties.Add(FPropertyAnimatorCoreData(TextRootComponent, LocationProperty, nullptr, UPropertyAnimatorTextResolver::StaticClass()));
-}
\ No newline at end of file
+	OutProperties.Add(FPropertyAnimatorCoreData(TextComponent, LocationProperty, nullptr, UPropertyAnimatorTextResolver::StaticClass()));
+}
diff --git a/Engine/Plugins/Experimental/PropertyAnimator/Source/PropertyAnimator/Private/Presets/PropertyAnimatorPresetTextRotation.cpp b/Engine/Plugins/Experimental/PropertyAnimator/Source/PropertyAnimator/Private/Presets/PropertyAnimatorPresetTextRotation.cpp
--- a/Engine/Plugins/Experimental/PropertyAnimator/Source/PropertyAnimator/Private/Presets/PropertyAnimatorPresetTextRotation.cpp
+++ b/Engine/Plugins/Experimental/PropertyAnimator/Source/PropertyAnimator/Private/Presets/PropertyAnimatorPresetTextRotation.cpp
@@ -3,27 +3,22 @@
 #include "Presets/PropertyAnimatorPresetTextRotation.h"
 
 #include "Animators/PropertyAnimatorCoreBase.h"
+#include "Characters/Text3DCharacterBase.h"
 #include "Properties/PropertyAnimatorTextResolver.h"
 #include "Properties/PropertyAnimatorFloatContext.h"
 #include "Text3DComponent.h"
 
 void UPropertyAnimatorPresetTextRotation::GetPresetProperties(const AActor* InActor, const UPropertyAnimatorCoreBase* InAnimator, TSet<FPropertyAnimatorCoreData>& OutProperties) const
 {
-	const UText3DComponent* TextComponent = InActor->FindComponentByClass<UText3DComponent>();
+	UText3DComponent* TextComponent = InActor->FindComponentByClass<UText3DComponent>();
 	if (!TextComponent)
 	{
 		return;
 	}
 
-	USceneComponent* TextRootComponent = TextComponent->GetChildComponent(0);
-	if (!TextRootComponent)
-	{
-		return;
-	}
-
-	const FName RotationPropertyName = USceneComponent::GetRelativeRotationPropertyName();
-	FProperty* RotationProperty = FindFProperty<FProperty>(USceneComponent::StaticClass(), RotationPropertyName);
+	const FName RotationPropertyName = UText3DCharacterBase::GetRelativeRotationPropertyName();
+	FProperty* RotationProperty = FindFProperty<FProperty>(UText3DCharacterBase::StaticClass(), RotationPropertyName);
 	check(RotationProperty);
 
-	OutProperties.Add(FPropertyAnimatorCoreData(TextRootComponent, RotationProperty, nullptr, UPropertyAnimatorTextResolver::StaticClass()));
-}
\ No newline at end of file
+	OutProperties.Add(FPropertyAnimatorCoreData(TextComponent, RotationProperty, nullptr, UPropertyAnimatorTextResolver::StaticClass()));
+}
diff --git a/Engine/Plugins/Experimental/PropertyAnimator/Source/PropertyAnimator/Private/Presets/PropertyAnimatorPresetTextScale.cpp b/Engine/Plugins/Experimental/PropertyAnimator/Source/PropertyAnimator/Private/Presets/PropertyAnimatorPresetTextScale.cpp
--- a/Engine/Plugins/Experimental/PropertyAnimator/Source/PropertyAnimator/Private/Presets/PropertyAnimatorPresetTextScale.cpp
+++ b/Engine/Plugins/Experimental/PropertyAnimator/Source/PropertyAnimator/Private/Presets/PropertyAnimatorPresetTextScale.cpp
@@ -3,27 +3,22 @@
 #include "Presets/PropertyAnimatorPresetTextScale.h"
 
 #include "Animators/PropertyAnimatorCoreBase.h"
+#include "Characters/Text3DCharacterBase.h"
 #include "Properties/PropertyAnimatorTextResolver.h"
 #include "Properties/PropertyAnimatorFloatContext.h"
 #include "Text3DComponent.h"
 
 void UPropertyAnimatorPresetTextScale::GetPresetProperties(const AActor* InActor, const UPropertyAnimatorCoreBase* InAnimator, TSet<FPropertyAnimatorCoreData>& OutProperties) const
 {
-	const UText3DComponent* TextComponent = InActor->FindComponentByClass<UText3DComponent>();
+	UText3DComponent* TextComponent = InActor->FindComponentByClass<UText3DComponent>();
 	if (!TextComponent)
 	{
 		return;
 	}
 
-	USceneComponent* TextRootComponent = TextComponent->GetChildComponent(0);
-	if (!TextRootComponent)
-	{
-		return;
-	}
-
-	const FName ScalePropertyName = USceneComponent::GetRelativeScale3DPropertyName();
-	FProperty* ScaleProperty = FindFProperty<FProperty>(USceneComponent::StaticClass(), ScalePropertyName);
+	const FName ScalePropertyName = UText3DCharacterBase::GetRelativeScalePropertyName();
+	FProperty* ScaleProperty = FindFProperty<FProperty>(UText3DCharacterBase::StaticClass(), ScalePropertyName);
 	check(ScaleProperty);
 
-	OutProperties.Add(FPropertyAnimatorCoreData(TextRootComponent, ScaleProperty, nullptr, UPropertyAnimatorTextResolver::StaticClass()));
-}
\ No newline at end of file
+	OutProperties.Add(FPropertyAnimatorCoreData(TextComponent, ScaleProperty, nullptr, UPropertyAnimatorTextResolver::StaticClass()));
+}
diff --git a/Engine/Plugins/Experimental/PropertyAnimator/Source/PropertyAnimator/Private/Properties/PropertyAnimatorTextResolver.cpp b/Engine/Plugins/Experimental/PropertyAnimator/Source/PropertyAnimator/Private/Properties/PropertyAnimatorTextResolver.cpp
--- a/Engine/Plugins/Experimental/PropertyAnimator/Source/PropertyAnimator/Private/Properties/PropertyAnimatorTextResolver.cpp
+++ b/Engine/Plugins/Experimental/PropertyAnimator/Source/PropertyAnimator/Private/Properties/PropertyAnimatorTextResolver.cpp
@@ -2,6 +2,8 @@
 
 #include "Properties/PropertyAnimatorTextResolver.h"
 
+#include "Characters/Text3DCharacterBase.h"
+#include "Misc/EnumerateRange.h"
 #include "Presets/PropertyAnimatorCorePresetArchive.h"
 #include "Text3DComponent.h"
 
@@ -60,10 +62,36 @@ void UPropertyAnimatorTextResolver::SetDirection(EPropertyAnimatorTextResolverRa
 	Direction = InDirection;
 }
 
-void UPropertyAnimatorTextResolver::GetResolvableProperties(const FPropertyAnimatorCoreData& InParentProperty, TSet<FPropertyAnimatorCoreData>& OutProperties)
+bool UPropertyAnimatorTextResolver::FixUpProperty(FPropertyAnimatorCoreData& InOldProperty)
 {
-	const AActor* Actor = InParentProperty.GetOwningActor();
-	if (!Actor || InParentProperty.IsResolved())
+	if (InOldProperty.GetPropertyResolver() != this)
+	{
+		return false;
+	}
+
+	TSet<FPropertyAnimatorCoreData> TemplateProperties;
+	const TArray<FName> SearchPaths = InOldProperty.GetChainPropertyNames();
+	GetTemplateProperties(InOldProperty.GetOwningActor(), TemplateProperties, &SearchPaths);
+
+	if (TemplateProperties.Num() != 1)
+	{
+		return false;
+	}
+
+	FPropertyAnimatorCoreData TemplateProperty = TemplateProperties.Array()[0];
+	if (TemplateProperty == InOldProperty)
+	{
+		return false;
+	}
+
+	InOldProperty = TemplateProperty;
+	return true;
+}
+
+void UPropertyAnimatorTextResolver::GetTemplateProperties(UObject* InContext, TSet<FPropertyAnimatorCoreData>& OutProperties, const TArray<FName>* InSearchPath)
+{
+	const AActor* Actor = Cast<AActor>(InContext);
+	if (!Actor)
 	{
 		return;
 	}
@@ -74,23 +102,57 @@ void UPropertyAnimatorTextResolver::GetResolvableProperties(const FPropertyAnima
 		return;
 	}
 
-	USceneComponent* TextRootComponent = TextComponent->GetChildComponent(0);
-	if (!TextRootComponent)
+	TArray<FPropertyAnimatorCoreData> TemplateProperties;
+
+	FProperty* RelativeLocation = FindFProperty<FProperty>(UText3DCharacterBase::StaticClass(), UText3DCharacterBase::GetRelativeLocationPropertyName());
+	FPropertyAnimatorCoreData RelativeLocationProperty(TextComponent, RelativeLocation, nullptr, GetClass());
+	TemplateProperties.Add(RelativeLocationProperty);
+
+	FProperty* RelativeRotation = FindFProperty<FProperty>(UText3DCharacterBase::StaticClass(), UText3DCharacterBase::GetRelativeRotationPropertyName());
+	FPropertyAnimatorCoreData RelativeRotationProperty(TextComponent, RelativeRotation, nullptr, GetClass());
+	TemplateProperties.Add(RelativeRotationProperty);
+
+	FProperty* RelativeScale = FindFProperty<FProperty>(UText3DCharacterBase::StaticClass(), UText3DCharacterBase::GetRelativeScalePropertyName());
+	FPropertyAnimatorCoreData RelativeScaleProperty(TextComponent, RelativeScale, nullptr, GetClass());
+	TemplateProperties.Add(RelativeScaleProperty);
+
+	if (!InSearchPath || InSearchPath->IsEmpty())
 	{
+		OutProperties.Append(TemplateProperties);
 		return;
 	}
 
-	FProperty* RelativeLocation = FindFProperty<FProperty>(TextRootComponent->GetClass(), TEXT("RelativeLocation"));
-	FPropertyAnimatorCoreData RelativeLocationProperty(TextRootComponent, RelativeLocation, nullptr, GetClass());
-	OutProperties.Add(RelativeLocationProperty);
+	while (!TemplateProperties.IsEmpty())
+	{
+		FPropertyAnimatorCoreData TemplateProperty = TemplateProperties.Pop();
+		const TArray<FName> ChainPropertyNames = TemplateProperty.GetChainPropertyNames();
+		for (TConstEnumerateRef<FName> TemplatePathPart : EnumerateRange(ChainPropertyNames))
+		{
+			if (!InSearchPath->IsValidIndex(TemplatePathPart.GetIndex()))
+			{
+				break;
+			}
+
+			const FName& SearchPathPart = (*InSearchPath)[TemplatePathPart.GetIndex()];
+			if (SearchPathPart != (*TemplatePathPart))
+			{
+				break;
+			}
 
-	FProperty* RelativeRotation = FindFProperty<FProperty>(TextRootComponent->GetClass(), TEXT("RelativeRotation"));
-	FPropertyAnimatorCoreData RelativeRotationProperty(TextRootComponent, RelativeRotation, nullptr, GetClass());
-	OutProperties.Add(RelativeRotationProperty);
+			if (ChainPropertyNames.Num() < InSearchPath->Num())
+			{
+				int32 DepthSearch = InSearchPath->Num() - ChainPropertyNames.Num();
+				TemplateProperties.Append(TemplateProperty.GetChildrenProperties(DepthSearch));
+				break;
+			}
 
-	FProperty* RelativeScale = FindFProperty<FProperty>(TextRootComponent->GetClass(), TEXT("RelativeScale3D"));
-	FPropertyAnimatorCoreData RelativeScaleProperty(TextRootComponent, RelativeScale, nullptr, GetClass());
-	OutProperties.Add(RelativeScaleProperty);
+			if (TemplatePathPart.GetIndex() == InSearchPath->Num() - 1)
+			{
+				OutProperties.Add(TemplateProperty);
+				break;
+			}
+		}
+	}
 }
 
 void UPropertyAnimatorTextResolver::ResolveProperties(const FPropertyAnimatorCoreData& InTemplateProperty, TArray<FPropertyAnimatorCoreData>& OutProperties, bool bInForEvaluation)
@@ -100,8 +162,8 @@ void UPropertyAnimatorTextResolver::ResolveProperties(const FPropertyAnimatorCor
 		return;
 	}
 
-	const USceneComponent* TextRootComponent = Cast<USceneComponent>(InTemplateProperty.GetOwningComponent());
-	if (!TextRootComponent)
+	const UText3DComponent* TextComponent = Cast<UText3DComponent>(InTemplateProperty.GetOwningComponent());
+	if (!TextComponent)
 	{
 		return;
 	}
@@ -109,16 +171,16 @@ void UPropertyAnimatorTextResolver::ResolveProperties(const FPropertyAnimatorCor
 	const TArray<FProperty*> ChainProperties = InTemplateProperty.GetChainProperties();
 
 	// Gather each character in the text
-	for (int32 ComponentIndex = 0; ComponentIndex < TextRootComponent->GetNumChildrenComponents(); ComponentIndex++)
+	for (int32 CharacterIndex = 0; CharacterIndex < TextComponent->GetCharacterCount(); CharacterIndex++)
 	{
-		USceneComponent* CharacterKerningComponent = TextRootComponent->GetChildComponent(ComponentIndex);
+		UText3DCharacterBase* Character = TextComponent->GetCharacter(CharacterIndex);
 
-		if (!CharacterKerningComponent)
+		if (!Character)
 		{
 			continue;
 		}
 
-		FPropertyAnimatorCoreData CharacterProperty(CharacterKerningComponent, ChainProperties);
+		FPropertyAnimatorCoreData CharacterProperty(Character, ChainProperties);
 		OutProperties.Add(CharacterProperty);
 	}
 
@@ -146,13 +208,6 @@ void UPropertyAnimatorTextResolver::ResolveProperties(const FPropertyAnimatorCor
 				EndPercentage = Temp;
 				OffsetPercentage *= -1;
 			}
-			else if (Direction == EPropertyAnimatorTextResolverRangeDirection::FromCenter)
-			{
-				constexpr float MidPercentage = 0.5f;
-				const float Expansion = EndPercentage / 2.f;
-				StartPercentage = MidPercentage - Expansion;
-				EndPercentage = MidPercentage + Expansion;
-			}
 
 			BeginIndex = StartPercentage * MaxIndex + OffsetPercentage * MaxIndex;
 			EndIndex = EndPercentage * MaxIndex + OffsetPercentage * MaxIndex;
@@ -172,13 +227,6 @@ void UPropertyAnimatorTextResolver::ResolveProperties(const FPropertyAnimatorCor
 				CharacterEnd = Temp;
 				CharacterOffset *= -1;
 			}
-			else if (Direction == EPropertyAnimatorTextResolverRangeDirection::FromCenter)
-			{
-				const int32 CharacterMid = MaxIndex / 2;
-				const int32 Expansion = CharacterEnd / 2;
-				CharacterStart = CharacterMid - Expansion;
-				CharacterEnd = CharacterMid + Expansion;
-			}
 
 			BeginIndex = CharacterStart + CharacterOffset;
 			EndIndex = CharacterEnd + CharacterOffset;
@@ -187,53 +235,43 @@ void UPropertyAnimatorTextResolver::ResolveProperties(const FPropertyAnimatorCor
 
 		case EPropertyAnimatorTextResolverRangeUnit::Word:
 		{
-			if (const UText3DComponent* TextComponent = TextRootComponent->GetTypedOuter<UText3DComponent>())
+			const FText3DStatistics& TextStats = TextComponent->GetStatistics();
+
+			if (TextStats.Words.IsEmpty())
 			{
-				const FText3DStatistics& TextStats = TextComponent->GetStatistics();
+				break;
+			}
 
-				if (TextStats.Words.IsEmpty())
-				{
-					break;
-				}
+			const int32 WordCount = TextStats.Words.Num();
+			int32 WordStart = WordStartIndex;
+			int32 WordEnd = WordEndIndex;
+			int32 WordOffset = WordOffsetIndex;
 
-				const int32 WordCount = TextStats.Words.Num();
-				int32 WordStart = WordStartIndex;
-				int32 WordEnd = WordEndIndex;
-				int32 WordOffset = WordOffsetIndex;
+			if (Direction == EPropertyAnimatorTextResolverRangeDirection::RightToLeft)
+			{
+				const int32 Temp = WordCount - WordStart;
+				WordStart = WordCount - WordEnd;
+				WordEnd = Temp;
+				WordOffset *= -1;
+			}
+
+			if (WordStart != WordEnd)
+			{
+				WordStart += WordOffset;
+				WordEnd += WordOffset - 1;
 
-				if (Direction == EPropertyAnimatorTextResolverRangeDirection::RightToLeft)
+				if (TextStats.Words.IsValidIndex(WordStart))
 				{
-					const int32 Temp = WordCount - WordStart;
-					WordStart = WordCount - WordEnd;
-					WordEnd = Temp;
-					WordOffset *= -1;
+					BeginIndex = TextStats.Words[WordStart].RenderRange.BeginIndex;
 				}
-				else if (Direction == EPropertyAnimatorTextResolverRangeDirection::FromCenter)
+
+				if (TextStats.Words.IsValidIndex(WordEnd))
 				{
-					const int32 WordMid = FMath::CeilToInt(WordCount / 2.f);
-					const int32 Expansion = FMath::CeilToInt(WordEnd / 2.f);
-					WordStart = WordMid - Expansion;
-					WordEnd = WordMid + Expansion;
+					EndIndex = TextStats.Words[WordEnd].RenderRange.EndIndex;
 				}
-
-				if (WordStart != WordEnd)
+				else if (WordEnd >= WordCount && WordStart < WordCount)
 				{
-					WordStart += WordOffset;
-					WordEnd += WordOffset - 1;
-
-					if (TextStats.Words.IsValidIndex(WordStart))
-					{
-						BeginIndex = TextStats.Words[WordStart].RenderRange.BeginIndex;
-					}
-
-					if (TextStats.Words.IsValidIndex(WordEnd))
-					{
-						EndIndex = TextStats.Words[WordEnd].RenderRange.EndIndex;
-					}
-					else if (WordEnd >= WordCount && WordStart < WordCount)
-					{
-						EndIndex = TextStats.Words.Last().RenderRange.EndIndex;
-					}
+					EndIndex = TextStats.Words.Last().RenderRange.EndIndex;
 				}
 			}
 		}
diff --git a/Engine/Plugins/Experimental/PropertyAnimator/Source/PropertyAnimator/Private/Properties/PropertyAnimatorTextResolver.h b/Engine/Plugins/Experimental/PropertyAnimator/Source/PropertyAnimator/Private/Properties/PropertyAnimatorTextResolver.h
--- a/Engine/Plugins/Experimental/PropertyAnimator/Source/PropertyAnimator/Private/Properties/PropertyAnimatorTextResolver.h
+++ b/Engine/Plugins/Experimental/PropertyAnimator/Source/PropertyAnimator/Private/Properties/PropertyAnimatorTextResolver.h
@@ -18,8 +18,7 @@ UENUM()
 enum class EPropertyAnimatorTextResolverRangeDirection : uint8
 {
 	LeftToRight,
-	RightToLeft,
-	FromCenter
+	RightToLeft
 };
 
 /**
@@ -105,7 +104,8 @@ public:
 	}
 
 	//~ Begin UPropertyAnimatorCoreResolver
-	virtual void GetResolvableProperties(const FPropertyAnimatorCoreData& InParentProperty, TSet<FPropertyAnimatorCoreData>& OutProperties) override;
+	virtual bool FixUpProperty(FPropertyAnimatorCoreData& InOldProperty) override;
+	virtual void GetTemplateProperties(UObject* InContext, TSet<FPropertyAnimatorCoreData>& OutProperties, const TArray<FName>* InSearchPath) override;
 	virtual void ResolveProperties(const FPropertyAnimatorCoreData& InTemplateProperty, TArray<FPropertyAnimatorCoreData>& OutProperties, bool bInForEvaluation) override;
 	virtual bool ImportPreset(const UPropertyAnimatorCorePresetBase* InPreset, const TSharedRef<FPropertyAnimatorCorePresetArchive>& InValue) override;
 	virtual bool ExportPreset(const UPropertyAnimatorCorePresetBase* InPreset, TSharedPtr<FPropertyAnimatorCorePresetArchive>& OutValue) const override;
diff --git a/Engine/Plugins/Experimental/PropertyAnimatorCore/Source/PropertyAnimatorCore/Private/Animators/PropertyAnimatorCoreBase.cpp b/Engine/Plugins/Experimental/PropertyAnimatorCore/Source/PropertyAnimatorCore/Private/Animators/PropertyAnimatorCoreBase.cpp
--- a/Engine/Plugins/Experimental/PropertyAnimatorCore/Source/PropertyAnimatorCore/Private/Animators/PropertyAnimatorCoreBase.cpp
+++ b/Engine/Plugins/Experimental/PropertyAnimatorCore/Source/PropertyAnimatorCore/Private/Animators/PropertyAnimatorCoreBase.cpp
@@ -10,6 +10,11 @@
 #include "Settings/PropertyAnimatorCoreSettings.h"
 #include "Subsystems/PropertyAnimatorCoreSubsystem.h"
 #include "TimeSources/PropertyAnimatorCoreTimeSourceBase.h"
+#include "UObject/ObjectSaveContext.h"
+
+#if WITH_EDITOR
+#include "Editor.h"
+#endif
 
 UPropertyAnimatorCoreBase::FOnAnimatorUpdated UPropertyAnimatorCoreBase::OnAnimatorAddedDelegate;
 UPropertyAnimatorCoreBase::FOnAnimatorUpdated UPropertyAnimatorCoreBase::OnAnimatorRemovedDelegate;
@@ -37,6 +42,7 @@ UPropertyAnimatorCoreBase::UPropertyAnimatorCoreBase()
 
 #if WITH_EDITOR
 		FCoreUObjectDelegates::OnObjectsReplaced.AddUObject(this, &UPropertyAnimatorCoreBase::OnObjectReplaced);
+		FEditorDelegates::PreSaveWorldWithContext.AddUObject(this, &UPropertyAnimatorCoreBase::OnPreSaveWorld);
 #endif
 	}
 }
@@ -122,6 +128,7 @@ void UPropertyAnimatorCoreBase::BeginDestroy()
 
 #if WITH_EDITOR
 	FCoreUObjectDelegates::OnObjectsReplaced.RemoveAll(this);
+	FEditorDelegates::PreSaveWorldWithContext.RemoveAll(this);
 #endif
 }
 
@@ -187,6 +194,18 @@ void UPropertyAnimatorCoreBase::PostEditUndo()
 	RestoreProperties(bForceReset);
 }
 
+void UPropertyAnimatorCoreBase::PreEditChange(FProperty* InPropertyAboutToChange)
+{
+	Super::PreEditChange(InPropertyAboutToChange);
+
+	if (InPropertyAboutToChange
+		&& InPropertyAboutToChange->GetFName() == GET_MEMBER_NAME_CHECKED(UPropertyAnimatorCoreBase, bAnimatorEnabled))
+	{
+		constexpr bool bForceReset = true;
+		RestoreProperties(bForceReset);
+	}
+}
+
 void UPropertyAnimatorCoreBase::PostEditChangeProperty(FPropertyChangedEvent& PropertyChangedEvent)
 {
 	Super::PostEditChangeProperty(PropertyChangedEvent);
@@ -610,6 +629,19 @@ void UPropertyAnimatorCoreBase::OnObjectReplaced(const TMap<UObject*, UObject*>&
 	}, bResolve);
 }
 
+#if WITH_EDITOR
+void UPropertyAnimatorCoreBase::OnPreSaveWorld(UWorld* InWorld, FObjectPreSaveContext InContext)
+{
+	if (InWorld != GetWorld())
+	{
+		return;
+	}
+
+	constexpr bool bForce = true;
+	RestoreProperties(bForce);
+}
+#endif
+
 void UPropertyAnimatorCoreBase::RestoreProperties(bool bInForce)
 {
 	constexpr bool bResolve = false;
diff --git a/Engine/Plugins/Experimental/PropertyAnimatorCore/Source/PropertyAnimatorCore/Private/Components/PropertyAnimatorCoreComponent.cpp b/Engine/Plugins/Experimental/PropertyAnimatorCore/Source/PropertyAnimatorCore/Private/Components/PropertyAnimatorCoreComponent.cpp
--- a/Engine/Plugins/Experimental/PropertyAnimatorCore/Source/PropertyAnimatorCore/Private/Components/PropertyAnimatorCoreComponent.cpp
+++ b/Engine/Plugins/Experimental/PropertyAnimatorCore/Source/PropertyAnimatorCore/Private/Components/PropertyAnimatorCoreComponent.cpp
@@ -497,6 +497,11 @@ bool UPropertyAnimatorCoreComponent::EvaluateAnimators()
 		return false;
 	}
 
+	if (UE::IsSavingPackage(this))
+	{
+		return true;
+	}
+
 	const UWorld* World = GetWorld();
 	const bool bIsSupportedWorld = IsValid(World) && (World->IsGameWorld() || World->IsEditorWorld());
 
diff --git a/Engine/Plugins/Experimental/PropertyAnimatorCore/Source/PropertyAnimatorCore/Private/Properties/Handlers/PropertyAnimatorCoreBoolHandler.cpp b/Engine/Plugins/Experimental/PropertyAnimatorCore/Source/PropertyAnimatorCore/Private/Properties/Handlers/PropertyAnimatorCoreBoolHandler.cpp
--- a/Engine/Plugins/Experimental/PropertyAnimatorCore/Source/PropertyAnimatorCore/Private/Properties/Handlers/PropertyAnimatorCoreBoolHandler.cpp
+++ b/Engine/Plugins/Experimental/PropertyAnimatorCore/Source/PropertyAnimatorCore/Private/Properties/Handlers/PropertyAnimatorCoreBoolHandler.cpp
@@ -38,4 +38,12 @@ bool UPropertyAnimatorCoreBoolHandler::SetValue(const FPropertyAnimatorCoreData&
 	InPropertyData.SetPropertyValuePtr(&bNewValue);
 
 	return true;
-}
\ No newline at end of file
+}
+
+bool UPropertyAnimatorCoreBoolHandler::GetDefaultValue(const FPropertyAnimatorCoreData& InPropertyData, FInstancedPropertyBag& OutValue)
+{
+	const FName PropertyHash(InPropertyData.GetLocatorPathHash());
+	OutValue.AddProperty(PropertyHash, EPropertyBagPropertyType::Bool);
+	OutValue.SetValueBool(PropertyHash, false);
+	return true;
+}
diff --git a/Engine/Plugins/Experimental/PropertyAnimatorCore/Source/PropertyAnimatorCore/Private/Properties/Handlers/PropertyAnimatorCoreBoolHandler.h b/Engine/Plugins/Experimental/PropertyAnimatorCore/Source/PropertyAnimatorCore/Private/Properties/Handlers/PropertyAnimatorCoreBoolHandler.h
--- a/Engine/Plugins/Experimental/PropertyAnimatorCore/Source/PropertyAnimatorCore/Private/Properties/Handlers/PropertyAnimatorCoreBoolHandler.h
+++ b/Engine/Plugins/Experimental/PropertyAnimatorCore/Source/PropertyAnimatorCore/Private/Properties/Handlers/PropertyAnimatorCoreBoolHandler.h
@@ -16,5 +16,6 @@ public:
 	virtual bool IsPropertySupported(const FPropertyAnimatorCoreData& InPropertyData) const override;
 	virtual bool GetValue(const FPropertyAnimatorCoreData& InPropertyData, FInstancedPropertyBag& OutValue) override;
 	virtual bool SetValue(const FPropertyAnimatorCoreData& InPropertyData, const FInstancedPropertyBag& InValue) override;
+	virtual bool GetDefaultValue(const FPropertyAnimatorCoreData& InPropertyData, FInstancedPropertyBag& OutValue) override;
 	//~ End UPropertyAnimatorCoreHandlerBase
 };
\ No newline at end of file
diff --git a/Engine/Plugins/Experimental/PropertyAnimatorCore/Source/PropertyAnimatorCore/Private/Properties/Handlers/PropertyAnimatorCoreDoubleHandler.cpp b/Engine/Plugins/Experimental/PropertyAnimatorCore/Source/PropertyAnimatorCore/Private/Properties/Handlers/PropertyAnimatorCoreDoubleHandler.cpp
--- a/Engine/Plugins/Experimental/PropertyAnimatorCore/Source/PropertyAnimatorCore/Private/Properties/Handlers/PropertyAnimatorCoreDoubleHandler.cpp
+++ b/Engine/Plugins/Experimental/PropertyAnimatorCore/Source/PropertyAnimatorCore/Private/Properties/Handlers/PropertyAnimatorCoreDoubleHandler.cpp
@@ -85,3 +85,11 @@ bool UPropertyAnimatorCoreDoubleHandler::SubtractValue(const FPropertyAnimatorCo
 
 	return true;
 }
+
+bool UPropertyAnimatorCoreDoubleHandler::GetDefaultValue(const FPropertyAnimatorCoreData& InPropertyData, FInstancedPropertyBag& OutValue)
+{
+	const FName PropertyHash(InPropertyData.GetLocatorPathHash());
+	OutValue.AddProperty(PropertyHash, EPropertyBagPropertyType::Double);
+	OutValue.SetValueDouble(PropertyHash, 0.0);
+	return true;
+}
diff --git a/Engine/Plugins/Experimental/PropertyAnimatorCore/Source/PropertyAnimatorCore/Private/Properties/Handlers/PropertyAnimatorCoreDoubleHandler.h b/Engine/Plugins/Experimental/PropertyAnimatorCore/Source/PropertyAnimatorCore/Private/Properties/Handlers/PropertyAnimatorCoreDoubleHandler.h
--- a/Engine/Plugins/Experimental/PropertyAnimatorCore/Source/PropertyAnimatorCore/Private/Properties/Handlers/PropertyAnimatorCoreDoubleHandler.h
+++ b/Engine/Plugins/Experimental/PropertyAnimatorCore/Source/PropertyAnimatorCore/Private/Properties/Handlers/PropertyAnimatorCoreDoubleHandler.h
@@ -19,5 +19,6 @@ public:
 	virtual bool IsAdditiveSupported() const override;
 	virtual bool AddValue(const FPropertyAnimatorCoreData& InPropertyData, const FInstancedPropertyBag& InValue) override;
 	virtual bool SubtractValue(const FPropertyAnimatorCoreData& InPropertyData, const FInstancedPropertyBag& InValue) override;
+	virtual bool GetDefaultValue(const FPropertyAnimatorCoreData& InPropertyData, FInstancedPropertyBag& OutValue) override;
 	//~ End UPropertyAnimatorCoreHandlerBase
 };
diff --git a/Engine/Plugins/Experimental/PropertyAnimatorCore/Source/PropertyAnimatorCore/Private/Properties/Handlers/PropertyAnimatorCoreFloatHandler.cpp b/Engine/Plugins/Experimental/PropertyAnimatorCore/Source/PropertyAnimatorCore/Private/Properties/Handlers/PropertyAnimatorCoreFloatHandler.cpp
--- a/Engine/Plugins/Experimental/PropertyAnimatorCore/Source/PropertyAnimatorCore/Private/Properties/Handlers/PropertyAnimatorCoreFloatHandler.cpp
+++ b/Engine/Plugins/Experimental/PropertyAnimatorCore/Source/PropertyAnimatorCore/Private/Properties/Handlers/PropertyAnimatorCoreFloatHandler.cpp
@@ -85,3 +85,11 @@ bool UPropertyAnimatorCoreFloatHandler::SubtractValue(const FPropertyAnimatorCor
 
 	return true;
 }
+
+bool UPropertyAnimatorCoreFloatHandler::GetDefaultValue(const FPropertyAnimatorCoreData& InPropertyData, FInstancedPropertyBag& OutValue)
+{
+	const FName PropertyHash(InPropertyData.GetLocatorPathHash());
+	OutValue.AddProperty(PropertyHash, EPropertyBagPropertyType::Float);
+	OutValue.SetValueFloat(PropertyHash, 0.f);
+	return true;
+}
diff --git a/Engine/Plugins/Experimental/PropertyAnimatorCore/Source/PropertyAnimatorCore/Private/Properties/Handlers/PropertyAnimatorCoreFloatHandler.h b/Engine/Plugins/Experimental/PropertyAnimatorCore/Source/PropertyAnimatorCore/Private/Properties/Handlers/PropertyAnimatorCoreFloatHandler.h
--- a/Engine/Plugins/Experimental/PropertyAnimatorCore/Source/PropertyAnimatorCore/Private/Properties/Handlers/PropertyAnimatorCoreFloatHandler.h
+++ b/Engine/Plugins/Experimental/PropertyAnimatorCore/Source/PropertyAnimatorCore/Private/Properties/Handlers/PropertyAnimatorCoreFloatHandler.h
@@ -19,5 +19,6 @@ public:
 	virtual bool IsAdditiveSupported() const override;
 	virtual bool AddValue(const FPropertyAnimatorCoreData& InPropertyData, const FInstancedPropertyBag& InValue) override;
 	virtual bool SubtractValue(const FPropertyAnimatorCoreData& InPropertyData, const FInstancedPropertyBag& InValue) override;
+	virtual bool GetDefaultValue(const FPropertyAnimatorCoreData& InPropertyData, FInstancedPropertyBag& OutValue) override;
 	//~ End UPropertyAnimatorCoreHandlerBase
 };
\ No newline at end of file
diff --git a/Engine/Plugins/Experimental/PropertyAnimatorCore/Source/PropertyAnimatorCore/Private/Properties/Handlers/PropertyAnimatorCoreIntHandler.cpp b/Engine/Plugins/Experimental/PropertyAnimatorCore/Source/PropertyAnimatorCore/Private/Properties/Handlers/PropertyAnimatorCoreIntHandler.cpp
--- a/Engine/Plugins/Experimental/PropertyAnimatorCore/Source/PropertyAnimatorCore/Private/Properties/Handlers/PropertyAnimatorCoreIntHandler.cpp
+++ b/Engine/Plugins/Experimental/PropertyAnimatorCore/Source/PropertyAnimatorCore/Private/Properties/Handlers/PropertyAnimatorCoreIntHandler.cpp
@@ -80,3 +80,11 @@ bool UPropertyAnimatorCoreIntHandler::SubtractValue(const FPropertyAnimatorCoreD
 
 	return true;
 }
+
+bool UPropertyAnimatorCoreIntHandler::GetDefaultValue(const FPropertyAnimatorCoreData& InPropertyData, FInstancedPropertyBag& OutValue)
+{
+	const FName PropertyHash(InPropertyData.GetLocatorPathHash());
+	OutValue.AddProperty(PropertyHash, EPropertyBagPropertyType::Int32);
+	OutValue.SetValueInt32(PropertyHash, 0);
+	return true;
+}
diff --git a/Engine/Plugins/Experimental/PropertyAnimatorCore/Source/PropertyAnimatorCore/Private/Properties/Handlers/PropertyAnimatorCoreIntHandler.h b/Engine/Plugins/Experimental/PropertyAnimatorCore/Source/PropertyAnimatorCore/Private/Properties/Handlers/PropertyAnimatorCoreIntHandler.h
--- a/Engine/Plugins/Experimental/PropertyAnimatorCore/Source/PropertyAnimatorCore/Private/Properties/Handlers/PropertyAnimatorCoreIntHandler.h
+++ b/Engine/Plugins/Experimental/PropertyAnimatorCore/Source/PropertyAnimatorCore/Private/Properties/Handlers/PropertyAnimatorCoreIntHandler.h
@@ -19,5 +19,6 @@ public:
 	virtual bool IsAdditiveSupported() const override;
 	virtual bool AddValue(const FPropertyAnimatorCoreData& InPropertyData, const FInstancedPropertyBag& InValue) override;
 	virtual bool SubtractValue(const FPropertyAnimatorCoreData& InPropertyData, const FInstancedPropertyBag& InValue) override;
+	virtual bool GetDefaultValue(const FPropertyAnimatorCoreData& InPropertyData, FInstancedPropertyBag& OutValue) override;
 	//~ End UPropertyAnimatorCoreHandlerBase
 };
\ No newline at end of file
diff --git a/Engine/Plugins/Experimental/PropertyAnimatorCore/Source/PropertyAnimatorCore/Private/Properties/Handlers/PropertyAnimatorCoreRotatorHandler.cpp b/Engine/Plugins/Experimental/PropertyAnimatorCore/Source/PropertyAnimatorCore/Private/Properties/Handlers/PropertyAnimatorCoreRotatorHandler.cpp
--- a/Engine/Plugins/Experimental/PropertyAnimatorCore/Source/PropertyAnimatorCore/Private/Properties/Handlers/PropertyAnimatorCoreRotatorHandler.cpp
+++ b/Engine/Plugins/Experimental/PropertyAnimatorCore/Source/PropertyAnimatorCore/Private/Properties/Handlers/PropertyAnimatorCoreRotatorHandler.cpp
@@ -80,3 +80,11 @@ bool UPropertyAnimatorCoreRotatorHandler::SubtractValue(const FPropertyAnimatorC
 
 	return true;
 }
+
+bool UPropertyAnimatorCoreRotatorHandler::GetDefaultValue(const FPropertyAnimatorCoreData& InPropertyData, FInstancedPropertyBag& OutValue)
+{
+	const FName PropertyHash(InPropertyData.GetLocatorPathHash());
+	OutValue.AddProperty(PropertyHash, EPropertyBagPropertyType::Struct, TBaseStructure<FRotator>::Get());
+	OutValue.SetValueStruct(PropertyHash, FRotator::ZeroRotator);
+	return true;
+}
diff --git a/Engine/Plugins/Experimental/PropertyAnimatorCore/Source/PropertyAnimatorCore/Private/Properties/Handlers/PropertyAnimatorCoreRotatorHandler.h b/Engine/Plugins/Experimental/PropertyAnimatorCore/Source/PropertyAnimatorCore/Private/Properties/Handlers/PropertyAnimatorCoreRotatorHandler.h
--- a/Engine/Plugins/Experimental/PropertyAnimatorCore/Source/PropertyAnimatorCore/Private/Properties/Handlers/PropertyAnimatorCoreRotatorHandler.h
+++ b/Engine/Plugins/Experimental/PropertyAnimatorCore/Source/PropertyAnimatorCore/Private/Properties/Handlers/PropertyAnimatorCoreRotatorHandler.h
@@ -19,5 +19,6 @@ public:
 	virtual bool IsAdditiveSupported() const override;
 	virtual bool AddValue(const FPropertyAnimatorCoreData& InPropertyData, const FInstancedPropertyBag& InValue) override;
 	virtual bool SubtractValue(const FPropertyAnimatorCoreData& InPropertyData, const FInstancedPropertyBag& InValue) override;
+	virtual bool GetDefaultValue(const FPropertyAnimatorCoreData& InPropertyData, FInstancedPropertyBag& OutValue) override;
 	//~ End UPropertyAnimatorCoreHandlerBase
 };
\ No newline at end of file
diff --git a/Engine/Plugins/Experimental/PropertyAnimatorCore/Source/PropertyAnimatorCore/Private/Properties/Handlers/PropertyAnimatorCoreStringHandler.cpp b/Engine/Plugins/Experimental/PropertyAnimatorCore/Source/PropertyAnimatorCore/Private/Properties/Handlers/PropertyAnimatorCoreStringHandler.cpp
--- a/Engine/Plugins/Experimental/PropertyAnimatorCore/Source/PropertyAnimatorCore/Private/Properties/Handlers/PropertyAnimatorCoreStringHandler.cpp
+++ b/Engine/Plugins/Experimental/PropertyAnimatorCore/Source/PropertyAnimatorCore/Private/Properties/Handlers/PropertyAnimatorCoreStringHandler.cpp
@@ -88,3 +88,11 @@ bool UPropertyAnimatorCoreStringHandler::SubtractValue(const FPropertyAnimatorCo
 
 	return true;
 }
+
+bool UPropertyAnimatorCoreStringHandler::GetDefaultValue(const FPropertyAnimatorCoreData& InPropertyData, FInstancedPropertyBag& OutValue)
+{
+	const FName PropertyHash(InPropertyData.GetLocatorPathHash());
+	OutValue.AddProperty(PropertyHash, EPropertyBagPropertyType::String);
+	OutValue.SetValueString(PropertyHash, FString());
+	return true;
+}
diff --git a/Engine/Plugins/Experimental/PropertyAnimatorCore/Source/PropertyAnimatorCore/Private/Properties/Handlers/PropertyAnimatorCoreStringHandler.h b/Engine/Plugins/Experimental/PropertyAnimatorCore/Source/PropertyAnimatorCore/Private/Properties/Handlers/PropertyAnimatorCoreStringHandler.h
--- a/Engine/Plugins/Experimental/PropertyAnimatorCore/Source/PropertyAnimatorCore/Private/Properties/Handlers/PropertyAnimatorCoreStringHandler.h
+++ b/Engine/Plugins/Experimental/PropertyAnimatorCore/Source/PropertyAnimatorCore/Private/Properties/Handlers/PropertyAnimatorCoreStringHandler.h
@@ -19,5 +19,6 @@ public:
 	virtual bool IsAdditiveSupported() const override;
 	virtual bool AddValue(const FPropertyAnimatorCoreData& InPropertyData, const FInstancedPropertyBag& InValue) override;
 	virtual bool SubtractValue(const FPropertyAnimatorCoreData& InPropertyData, const FInstancedPropertyBag& InValue) override;
+	virtual bool GetDefaultValue(const FPropertyAnimatorCoreData& InPropertyData, FInstancedPropertyBag& OutValue) override;
 	//~ End UPropertyAnimatorCoreHandlerBase
 };
\ No newline at end of file
diff --git a/Engine/Plugins/Experimental/PropertyAnimatorCore/Source/PropertyAnimatorCore/Private/Properties/Handlers/PropertyAnimatorCoreTextHandler.cpp b/Engine/Plugins/Experimental/PropertyAnimatorCore/Source/PropertyAnimatorCore/Private/Properties/Handlers/PropertyAnimatorCoreTextHandler.cpp
--- a/Engine/Plugins/Experimental/PropertyAnimatorCore/Source/PropertyAnimatorCore/Private/Properties/Handlers/PropertyAnimatorCoreTextHandler.cpp
+++ b/Engine/Plugins/Experimental/PropertyAnimatorCore/Source/PropertyAnimatorCore/Private/Properties/Handlers/PropertyAnimatorCoreTextHandler.cpp
@@ -91,3 +91,11 @@ bool UPropertyAnimatorCoreTextHandler::SubtractValue(const FPropertyAnimatorCore
 
 	return true;
 }
+
+bool UPropertyAnimatorCoreTextHandler::GetDefaultValue(const FPropertyAnimatorCoreData& InPropertyData, FInstancedPropertyBag& OutValue)
+{
+	const FName PropertyHash(InPropertyData.GetLocatorPathHash());
+	OutValue.AddProperty(PropertyHash, EPropertyBagPropertyType::Text);
+	OutValue.SetValueText(PropertyHash, FText::GetEmpty());
+	return true;
+}
diff --git a/Engine/Plugins/Experimental/PropertyAnimatorCore/Source/PropertyAnimatorCore/Private/Properties/Handlers/PropertyAnimatorCoreTextHandler.h b/Engine/Plugins/Experimental/PropertyAnimatorCore/Source/PropertyAnimatorCore/Private/Properties/Handlers/PropertyAnimatorCoreTextHandler.h
--- a/Engine/Plugins/Experimental/PropertyAnimatorCore/Source/PropertyAnimatorCore/Private/Properties/Handlers/PropertyAnimatorCoreTextHandler.h
+++ b/Engine/Plugins/Experimental/PropertyAnimatorCore/Source/PropertyAnimatorCore/Private/Properties/Handlers/PropertyAnimatorCoreTextHandler.h
@@ -19,5 +19,6 @@ public:
 	virtual bool IsAdditiveSupported() const override;
 	virtual bool AddValue(const FPropertyAnimatorCoreData& InPropertyData, const FInstancedPropertyBag& InValue) override;
 	virtual bool SubtractValue(const FPropertyAnimatorCoreData& InPropertyData, const FInstancedPropertyBag& InValue) override;
+	virtual bool GetDefaultValue(const FPropertyAnimatorCoreData& InPropertyData, FInstancedPropertyBag& OutValue) override;
 	//~ End UPropertyAnimatorCoreHandlerBase
 };
\ No newline at end of file
diff --git a/Engine/Plugins/Experimental/PropertyAnimatorCore/Source/PropertyAnimatorCore/Private/Properties/Handlers/PropertyAnimatorCoreVectorHandler.cpp b/Engine/Plugins/Experimental/PropertyAnimatorCore/Source/PropertyAnimatorCore/Private/Properties/Handlers/PropertyAnimatorCoreVectorHandler.cpp
--- a/Engine/Plugins/Experimental/PropertyAnimatorCore/Source/PropertyAnimatorCore/Private/Properties/Handlers/PropertyAnimatorCoreVectorHandler.cpp
+++ b/Engine/Plugins/Experimental/PropertyAnimatorCore/Source/PropertyAnimatorCore/Private/Properties/Handlers/PropertyAnimatorCoreVectorHandler.cpp
@@ -80,3 +80,11 @@ bool UPropertyAnimatorCoreVectorHandler::SubtractValue(const FPropertyAnimatorCo
 
 	return true;
 }
+
+bool UPropertyAnimatorCoreVectorHandler::GetDefaultValue(const FPropertyAnimatorCoreData& InPropertyData, FInstancedPropertyBag& OutValue)
+{
+	const FName PropertyHash(InPropertyData.GetLocatorPathHash());
+	OutValue.AddProperty(PropertyHash, EPropertyBagPropertyType::Struct, TBaseStructure<FVector>::Get());
+	OutValue.SetValueStruct(PropertyHash, FVector::ZeroVector);
+	return true;
+}
diff --git a/Engine/Plugins/Experimental/PropertyAnimatorCore/Source/PropertyAnimatorCore/Private/Properties/Handlers/PropertyAnimatorCoreVectorHandler.h b/Engine/Plugins/Experimental/PropertyAnimatorCore/Source/PropertyAnimatorCore/Private/Properties/Handlers/PropertyAnimatorCoreVectorHandler.h
--- a/Engine/Plugins/Experimental/PropertyAnimatorCore/Source/PropertyAnimatorCore/Private/Properties/Handlers/PropertyAnimatorCoreVectorHandler.h
+++ b/Engine/Plugins/Experimental/PropertyAnimatorCore/Source/PropertyAnimatorCore/Private/Properties/Handlers/PropertyAnimatorCoreVectorHandler.h
@@ -19,5 +19,6 @@ public:
 	virtual bool IsAdditiveSupported() const override;
 	virtual bool AddValue(const FPropertyAnimatorCoreData& InPropertyData, const FInstancedPropertyBag& InValue) override;
 	virtual bool SubtractValue(const FPropertyAnimatorCoreData& InPropertyData, const FInstancedPropertyBag& InValue) override;
+	virtual bool GetDefaultValue(const FPropertyAnimatorCoreData& InPropertyData, FInstancedPropertyBag& OutValue) override;
 	//~ End UPropertyAnimatorCoreHandlerBase
 };
\ No newline at end of file
diff --git a/Engine/Plugins/Experimental/PropertyAnimatorCore/Source/PropertyAnimatorCore/Private/Properties/PropertyAnimatorCoreContext.cpp b/Engine/Plugins/Experimental/PropertyAnimatorCore/Source/PropertyAnimatorCore/Private/Properties/PropertyAnimatorCoreContext.cpp
--- a/Engine/Plugins/Experimental/PropertyAnimatorCore/Source/PropertyAnimatorCore/Private/Properties/PropertyAnimatorCoreContext.cpp
+++ b/Engine/Plugins/Experimental/PropertyAnimatorCore/Source/PropertyAnimatorCore/Private/Properties/PropertyAnimatorCoreContext.cpp
@@ -4,6 +4,7 @@
 
 #include "Containers/Ticker.h"
 #include "Logs/PropertyAnimatorCoreLog.h"
+#include "Misc/EnumerateRange.h"
 #include "Properties/PropertyAnimatorCoreResolver.h"
 #include "Properties/Converters/PropertyAnimatorCoreConverterBase.h"
 #include "Properties/Handlers/PropertyAnimatorCoreHandlerBase.h"
@@ -125,14 +126,17 @@ void UPropertyAnimatorCoreContext::PostLoad()
 	CheckEditConverterRule();
 	CheckEditResolver();
 
-	AnimatedProperty.GeneratePropertyPath();
-
-	if (AnimatedProperty.IsResolvable())
+	if (UPropertyAnimatorCoreResolver* PropertyResolver = AnimatedProperty.GetPropertyResolver())
 	{
-		DeltaPropertyValues.Reset();
-		OriginalPropertyValues.Reset();
+		TSet<FPropertyAnimatorCoreData> ResolvedProperties;
+		if (PropertyResolver->FixUpProperty(AnimatedProperty))
+		{
+			UE_LOG(LogPropertyAnimatorCore, Log, TEXT("Fixed up property %s using %s resolver"), *AnimatedProperty.GetPropertyDisplayName(), *PropertyResolver->GetResolverName().ToString())
+		}
 	}
 
+	AnimatedProperty.GeneratePropertyPath();
+
 	for (const FPropertyAnimatorCoreData& Property : ResolveProperty(/** Evaluation */false))
 	{
 		const FName OldPropertyPath(Property.GetPathHash());
@@ -405,11 +409,8 @@ void UPropertyAnimatorCoreContext::Restore()
 	{
 		for (const FPropertyAnimatorCoreData& ResolvedProperty : ResolveProperty(false))
 		{
-			// Reset original value
-			if (Handler->SetValue(ResolvedProperty, OriginalPropertyValues))
-			{
-				OriginalPropertyValues.RemovePropertyByName(ResolvedProperty.GetLocatorPathHash());
-			}
+			// Set original value
+			Handler->SetValue(ResolvedProperty, OriginalPropertyValues);
 		}
 
 		OriginalPropertyValues.Reset();
@@ -419,10 +420,7 @@ void UPropertyAnimatorCoreContext::Restore()
 		for (const FPropertyAnimatorCoreData& ResolvedProperty : ResolveProperty(false))
 		{
 			// Subtract delta value
-			if (Handler->SubtractValue(ResolvedProperty, DeltaPropertyValues))
-			{
-				DeltaPropertyValues.RemovePropertyByName(ResolvedProperty.GetLocatorPathHash());
-			}
+			Handler->SubtractValue(ResolvedProperty, DeltaPropertyValues);
 		}
 	}
 
@@ -454,6 +452,9 @@ void UPropertyAnimatorCoreContext::Save()
 		{
 			const FProperty* Property = PropertyData.GetLeafProperty();
 			DeltaPropertyValues.AddProperty(Name, Property);
+
+			// Save default value (zero)
+			Handler->GetDefaultValue(PropertyData, DeltaPropertyValues);
 		}
 	}
 }
diff --git a/Engine/Plugins/Experimental/PropertyAnimatorCore/Source/PropertyAnimatorCore/Private/Properties/PropertyAnimatorCoreData.cpp b/Engine/Plugins/Experimental/PropertyAnimatorCore/Source/PropertyAnimatorCore/Private/Properties/PropertyAnimatorCoreData.cpp
--- a/Engine/Plugins/Experimental/PropertyAnimatorCore/Source/PropertyAnimatorCore/Private/Properties/PropertyAnimatorCoreData.cpp
+++ b/Engine/Plugins/Experimental/PropertyAnimatorCore/Source/PropertyAnimatorCore/Private/Properties/PropertyAnimatorCoreData.cpp
@@ -193,9 +193,8 @@ FPropertyAnimatorCoreData::FPropertyAnimatorCoreData(AActor* InActor, const FStr
 		// Use resolver if owner was not found
 		if (!FoundObject && Resolver && MemberProperty)
 		{
-			const FPropertyAnimatorCoreData ActorData(InActor, nullptr, nullptr);
 			TSet<FPropertyAnimatorCoreData> ResolvableProperties;
-			Resolver->GetResolvableProperties(ActorData, ResolvableProperties);
+			Resolver->GetTemplateProperties(InActor, ResolvableProperties);
 
 			for (const FPropertyAnimatorCoreData& ResolvableProperty : ResolvableProperties)
 			{
@@ -460,6 +459,22 @@ TArray<FProperty*> FPropertyAnimatorCoreData::GetChainProperties() const
 	return Properties;
 }
 
+TArray<FName> FPropertyAnimatorCoreData::GetChainPropertyNames() const
+{
+	TArray<FName> Properties;
+
+	Algo::Transform(
+		ChainProperties
+		, Properties
+		, [](const TFieldPath<FProperty>& InProperty)
+		{
+			return InProperty->GetFName();
+		}
+	);
+
+	return Properties;
+}
+
 bool FPropertyAnimatorCoreData::HasSetter() const
 {
 	FPropertyAnimatorCoreData* This = const_cast<FPropertyAnimatorCoreData*>(this);
diff --git a/Engine/Plugins/Experimental/PropertyAnimatorCore/Source/PropertyAnimatorCore/Private/Subsystems/PropertyAnimatorCoreSubsystem.cpp b/Engine/Plugins/Experimental/PropertyAnimatorCore/Source/PropertyAnimatorCore/Private/Subsystems/PropertyAnimatorCoreSubsystem.cpp
--- a/Engine/Plugins/Experimental/PropertyAnimatorCore/Source/PropertyAnimatorCore/Private/Subsystems/PropertyAnimatorCoreSubsystem.cpp
+++ b/Engine/Plugins/Experimental/PropertyAnimatorCore/Source/PropertyAnimatorCore/Private/Subsystems/PropertyAnimatorCoreSubsystem.cpp
@@ -936,17 +936,6 @@ bool UPropertyAnimatorCoreSubsystem::IsResolverClassRegistered(const UClass* InR
 	return false;
 }
 
-void UPropertyAnimatorCoreSubsystem::GetResolvableProperties(const FPropertyAnimatorCoreData& InPropertyData, TSet<FPropertyAnimatorCoreData>& OutProperties) const
-{
-	for (const TWeakObjectPtr<UPropertyAnimatorCoreResolver>& ResolverWeak : ResolversWeak)
-	{
-		if (UPropertyAnimatorCoreResolver* Resolver = ResolverWeak.Get())
-		{
-			Resolver->GetResolvableProperties(InPropertyData, OutProperties);
-		}
-	}
-}
-
 bool UPropertyAnimatorCoreSubsystem::RegisterTimeSourceClass(UClass* InTimeSourceClass)
 {
 	if (!IsValid(InTimeSourceClass))
diff --git a/Engine/Plugins/Experimental/PropertyAnimatorCore/Source/PropertyAnimatorCore/Public/Animators/PropertyAnimatorCoreBase.h b/Engine/Plugins/Experimental/PropertyAnimatorCore/Source/PropertyAnimatorCore/Public/Animators/PropertyAnimatorCoreBase.h
--- a/Engine/Plugins/Experimental/PropertyAnimatorCore/Source/PropertyAnimatorCore/Public/Animators/PropertyAnimatorCoreBase.h
+++ b/Engine/Plugins/Experimental/PropertyAnimatorCore/Source/PropertyAnimatorCore/Public/Animators/PropertyAnimatorCoreBase.h
@@ -8,6 +8,7 @@
 #include "UObject/Object.h"
 #include "PropertyAnimatorCoreBase.generated.h"
 
+class FObjectPreSaveContext;
 class UPropertyAnimatorCoreComponent;
 class UPropertyAnimatorCoreTimeSourceBase;
 
@@ -239,6 +240,7 @@ protected:
 #if WITH_EDITOR
 	PROPERTYANIMATORCORE_API virtual void PreEditUndo() override;
 	PROPERTYANIMATORCORE_API virtual void PostEditUndo() override;
+	PROPERTYANIMATORCORE_API virtual void PreEditChange(FProperty* InPropertyAboutToChange) override;
 	PROPERTYANIMATORCORE_API virtual void PostEditChangeProperty(FPropertyChangedEvent& PropertyChangedEvent) override;
 #endif
 	//~ End UObject
@@ -334,7 +336,10 @@ private:
 
 	void OnObjectReplaced(const TMap<UObject*, UObject*>& InReplacementMap);
 
-	void OnPropertyGroupsChanged();
+#if WITH_EDITOR
+	/** Needed to restore properties and stop animation before world is saved */
+	void OnPreSaveWorld(UWorld* InWorld, FObjectPreSaveContext InContext);
+#endif
 
 	void OnAnimatorEnabledChanged(EPropertyAnimatorCoreUpdateEvent InType);
 
diff --git a/Engine/Plugins/Experimental/PropertyAnimatorCore/Source/PropertyAnimatorCore/Public/Properties/Handlers/PropertyAnimatorCoreHandlerBase.h b/Engine/Plugins/Experimental/PropertyAnimatorCore/Source/PropertyAnimatorCore/Public/Properties/Handlers/PropertyAnimatorCoreHandlerBase.h
--- a/Engine/Plugins/Experimental/PropertyAnimatorCore/Source/PropertyAnimatorCore/Public/Properties/Handlers/PropertyAnimatorCoreHandlerBase.h
+++ b/Engine/Plugins/Experimental/PropertyAnimatorCore/Source/PropertyAnimatorCore/Public/Properties/Handlers/PropertyAnimatorCoreHandlerBase.h
@@ -53,5 +53,11 @@ public:
 	{
 		return false;
 	}
+
+	/** Get the default value for this property, init to zero */
+	virtual bool GetDefaultValue(const FPropertyAnimatorCoreData& InPropertyData, FInstancedPropertyBag& OutValue)
+	{
+		return false;
+	}
 };
 
diff --git a/Engine/Plugins/Experimental/PropertyAnimatorCore/Source/PropertyAnimatorCore/Public/Properties/PropertyAnimatorCoreData.h b/Engine/Plugins/Experimental/PropertyAnimatorCore/Source/PropertyAnimatorCore/Public/Properties/PropertyAnimatorCoreData.h
--- a/Engine/Plugins/Experimental/PropertyAnimatorCore/Source/PropertyAnimatorCore/Public/Properties/PropertyAnimatorCoreData.h
+++ b/Engine/Plugins/Experimental/PropertyAnimatorCore/Source/PropertyAnimatorCore/Public/Properties/PropertyAnimatorCoreData.h
@@ -123,6 +123,9 @@ struct FPropertyAnimatorCoreData
 	/** The chain properties from member to inner property */
 	PROPERTYANIMATORCORE_API TArray<FProperty*> GetChainProperties() const;
 
+	/** The chain property names from member to inner property */
+	PROPERTYANIMATORCORE_API TArray<FName> GetChainPropertyNames() const;
+
 	/** Checks if the property is settable via setter */
 	PROPERTYANIMATORCORE_API bool HasSetter() const;
 
@@ -161,7 +164,7 @@ struct FPropertyAnimatorCoreData
 	}
 
 	/** Gets the children of this property */
-	TArray<FPropertyAnimatorCoreData> GetChildrenProperties(int32 InDepthSearch = 3) const;
+	PROPERTYANIMATORCORE_API TArray<FPropertyAnimatorCoreData> GetChildrenProperties(int32 InDepthSearch = 3) const;
 
 	/** Checks if this property contains a specific type */
 	template<typename InPropertyClass
diff --git a/Engine/Plugins/Experimental/PropertyAnimatorCore/Source/PropertyAnimatorCore/Public/Properties/PropertyAnimatorCoreResolver.h b/Engine/Plugins/Experimental/PropertyAnimatorCore/Source/PropertyAnimatorCore/Public/Properties/PropertyAnimatorCoreResolver.h
--- a/Engine/Plugins/Experimental/PropertyAnimatorCore/Source/PropertyAnimatorCore/Public/Properties/PropertyAnimatorCoreResolver.h
+++ b/Engine/Plugins/Experimental/PropertyAnimatorCore/Source/PropertyAnimatorCore/Public/Properties/PropertyAnimatorCoreResolver.h
@@ -27,10 +27,16 @@ public:
 		: ResolverName(InResolverName)
 	{}
 
-	/** Get properties found inside parent property */
-	virtual void GetResolvableProperties(const FPropertyAnimatorCoreData& InParentProperty, TSet<FPropertyAnimatorCoreData>& OutProperties) {}
+	/** Tries to resolve old property against new one */
+	virtual bool FixUpProperty(FPropertyAnimatorCoreData& InOldProperty)
+	{
+		return false;
+	}
+
+	/** Get template properties based on context */
+	virtual void GetTemplateProperties(UObject* InContext, TSet<FPropertyAnimatorCoreData>& OutProperties, const TArray<FName>* InSearchPath = nullptr) {}
 
-	/** Called when we actually need the underlying properties */
+	/** Called when we actually need the underlying properties of a template property */
 	virtual void ResolveProperties(const FPropertyAnimatorCoreData& InTemplateProperty, TArray<FPropertyAnimatorCoreData>& OutProperties, bool bInForEvaluation) {}
 
 	FName GetResolverName() const
diff --git a/Engine/Plugins/Experimental/PropertyAnimatorCore/Source/PropertyAnimatorCore/Public/Subsystems/PropertyAnimatorCoreSubsystem.h b/Engine/Plugins/Experimental/PropertyAnimatorCore/Source/PropertyAnimatorCore/Public/Subsystems/PropertyAnimatorCoreSubsystem.h
--- a/Engine/Plugins/Experimental/PropertyAnimatorCore/Source/PropertyAnimatorCore/Public/Subsystems/PropertyAnimatorCoreSubsystem.h
+++ b/Engine/Plugins/Experimental/PropertyAnimatorCore/Source/PropertyAnimatorCore/Public/Subsystems/PropertyAnimatorCoreSubsystem.h
@@ -85,9 +85,6 @@ public:
 	/** Is this resolver registered */
 	PROPERTYANIMATORCORE_API bool IsResolverClassRegistered(const UClass* InResolverClass) const;
 
-	/** Get all resolvable properties for an object property */
-	PROPERTYANIMATORCORE_API void GetResolvableProperties(const FPropertyAnimatorCoreData& InPropertyData, TSet<FPropertyAnimatorCoreData>& OutProperties) const;
-
 	/** Register a time source class to control clock for animators */
 	PROPERTYANIMATORCORE_API bool RegisterTimeSourceClass(UClass* InTimeSourceClass);
 
diff --git a/Engine/Plugins/Experimental/Text3D/Source/Text3D/Private/Characters/Text3DCharacterBase.cpp b/Engine/Plugins/Experimental/Text3D/Source/Text3D/Private/Characters/Text3DCharacterBase.cpp
--- a/Engine/Plugins/Experimental/Text3D/Source/Text3D/Private/Characters/Text3DCharacterBase.cpp
+++ b/Engine/Plugins/Experimental/Text3D/Source/Text3D/Private/Characters/Text3DCharacterBase.cpp
@@ -4,6 +4,21 @@
 
 #include "Text3DComponent.h"
 
+FName UText3DCharacterBase::GetRelativeLocationPropertyName()
+{
+	return GET_MEMBER_NAME_CHECKED(UText3DCharacterBase, RelativeLocation);
+}
+
+FName UText3DCharacterBase::GetRelativeRotationPropertyName()
+{
+	return GET_MEMBER_NAME_CHECKED(UText3DCharacterBase, RelativeRotation);
+}
+
+FName UText3DCharacterBase::GetRelativeScalePropertyName()
+{
+	return GET_MEMBER_NAME_CHECKED(UText3DCharacterBase, RelativeScale);
+}
+
 FTransform& UText3DCharacterBase::GetTransform(bool bInReset)
 {
 	if (bInReset)
@@ -25,7 +40,7 @@ void UText3DCharacterBase::SetRelativeLocation(const FVector& InLocation)
 	{
 		return;
 	}
-	
+
 	RelativeLocation = InLocation;
 	OnCharacterDataChanged();
 }
@@ -36,7 +51,7 @@ void UText3DCharacterBase::SetRelativeRotation(const FRotator& InRotation)
 	{
 		return;
 	}
-	
+
 	RelativeRotation = InRotation;
 	OnCharacterDataChanged();
 }
@@ -47,7 +62,7 @@ void UText3DCharacterBase::SetRelativeScale(const FVector& InScale)
 	{
 		return;
 	}
-	
+
 	RelativeScale = InScale;
 	OnCharacterDataChanged();
 }
@@ -82,6 +97,7 @@ void UText3DCharacterBase::OnCharacterDataChanged() const
 {
 	if (UText3DComponent* Component = GetTypedOuter<UText3DComponent>())
 	{
-		Component->RequestUpdate(EText3DRendererFlags::Layout);
+		constexpr bool bImmediate = true;
+		Component->RequestUpdate(EText3DRendererFlags::Layout, bImmediate);
 	}
 }
diff --git a/Engine/Plugins/Experimental/Text3D/Source/Text3D/Public/Characters/Text3DCharacterBase.h b/Engine/Plugins/Experimental/Text3D/Source/Text3D/Public/Characters/Text3DCharacterBase.h
--- a/Engine/Plugins/Experimental/Text3D/Source/Text3D/Public/Characters/Text3DCharacterBase.h
+++ b/Engine/Plugins/Experimental/Text3D/Source/Text3D/Public/Characters/Text3DCharacterBase.h
@@ -7,12 +7,16 @@
 #include "Text3DCharacterBase.generated.h"
 
 /** Holds data for a single character in Text3D */
-UCLASS(AutoExpandCategories=(Character))
+UCLASS(MinimalAPI, AutoExpandCategories=(Character))
 class UText3DCharacterBase : public UObject
 {
 	GENERATED_BODY()
 
 public:
+	static TEXT3D_API FName GetRelativeLocationPropertyName();
+	static TEXT3D_API FName GetRelativeRotationPropertyName();
+	static TEXT3D_API FName GetRelativeScalePropertyName();
+
 	TEXT3D_API FTransform& GetTransform(bool bInReset);
 
 #if WITH_EDITORONLY_DATA
@@ -20,6 +24,11 @@ public:
 	{
 		Character = InCharacter;
 	}
+
+	const FString& GetCharacter() const
+	{
+		return Character;
+	}
 #endif
 	
 	void SetCharacterMetadata(const FText3DCharacterMetadata& InMetadata);
@@ -32,24 +41,27 @@ public:
 		return Metadata;
 	}
 
-	void SetRelativeLocation(const FVector& InLocation);
+	UFUNCTION()
+	TEXT3D_API void SetRelativeLocation(const FVector& InLocation);
 	const FVector& GetRelativeLocation() const
 	{
 		return RelativeLocation;
 	}
 
-	void SetRelativeRotation(const FRotator& InRotation);
+	UFUNCTION()
+	TEXT3D_API void SetRelativeRotation(const FRotator& InRotation);
 	const FRotator& GetRelativeRotation() const
 	{
 		return RelativeRotation;
 	}
 
-	void SetRelativeScale(const FVector& InScale);
+	UFUNCTION()
+	TEXT3D_API void SetRelativeScale(const FVector& InScale);
 	const FVector& GetRelativeScale() const
 	{
 		return RelativeScale;
 	}
-	
+
 	/** Get character custom kerning if supported */
 	virtual TOptional<float> GetCharacterKerning() const
 	{
@@ -72,13 +84,13 @@ protected:
 	FString Character;
 #endif
 
-	UPROPERTY(EditInstanceOnly, Getter, Setter, Category="Character", meta=(AllowPrivateAccess = "true"))
+	UPROPERTY(EditInstanceOnly, Setter, Getter, Category="Character", meta=(AllowPrivateAccess = "true"))
 	FVector RelativeLocation = FVector::ZeroVector;
 	
-	UPROPERTY(EditInstanceOnly, Getter, Setter, Category="Character", meta=(AllowPrivateAccess = "true"))
+	UPROPERTY(EditInstanceOnly, Setter, Getter, Category="Character", meta=(AllowPrivateAccess = "true"))
 	FRotator RelativeRotation = FRotator::ZeroRotator;
 	
-	UPROPERTY(EditInstanceOnly, Getter, Setter, Category="Character", meta=(Delta="0.05", MotionDesignVectorWidget, AllowPreserveRatio="XYZ", ClampMin="0", AllowPrivateAccess = "true"))
+	UPROPERTY(EditInstanceOnly, Setter, Getter, Category="Character", meta=(Delta="0.05", MotionDesignVectorWidget, AllowPreserveRatio="XYZ", ClampMin="0", AllowPrivateAccess = "true"))
 	FVector RelativeScale = FVector::OneVector;
 
 	/** Final transform after all extensions are applied */
