MotionDesign : PropertyAnimator
- Reworking text resolver to better handle text character range based on options
- Adding support for text character visibility

diff --git a/Engine/Plugins/Experimental/PropertyAnimator/Source/PropertyAnimator/Private/Presets/PropertyAnimatorPresetTextVisibility.cpp b/Engine/Plugins/Experimental/PropertyAnimator/Source/PropertyAnimator/Private/Presets/PropertyAnimatorPresetTextVisibility.cpp
--- a/Engine/Plugins/Experimental/PropertyAnimator/Source/PropertyAnimator/Private/Presets/PropertyAnimatorPresetTextVisibility.cpp
+++ b/Engine/Plugins/Experimental/PropertyAnimator/Source/PropertyAnimator/Private/Presets/PropertyAnimatorPresetTextVisibility.cpp
@@ -3,27 +3,21 @@
 #include "Presets/PropertyAnimatorPresetTextVisibility.h"
 
 #include "Animators/PropertyAnimatorCoreBase.h"
+#include "Characters/Text3DCharacterBase.h"
 #include "Properties/PropertyAnimatorTextResolver.h"
-#include "Properties/Converters/PropertyAnimatorCoreConverterTraits.h"
 #include "Text3DComponent.h"
 
 void UPropertyAnimatorPresetTextVisibility::GetPresetProperties(const AActor* InActor, const UPropertyAnimatorCoreBase* InAnimator, TSet<FPropertyAnimatorCoreData>& OutProperties) const
 {
-	const UText3DComponent* TextComponent = InActor->FindComponentByClass<UText3DComponent>();
+	UText3DComponent* TextComponent = InActor->FindComponentByClass<UText3DComponent>();
 	if (!TextComponent)
 	{
 		return;
 	}
 
-	USceneComponent* TextRootComponent = TextComponent->GetChildComponent(0);
-	if (!TextRootComponent)
-	{
-		return;
-	}
-
-	const FName VisibilityPropertyName = USceneComponent::GetVisiblePropertyName();
-	FProperty* VisibilityProperty = FindFProperty<FProperty>(USceneComponent::StaticClass(), VisibilityPropertyName);
+	const FName VisibilityPropertyName = UText3DCharacterBase::GetVisiblePropertyName();
+	FProperty* VisibilityProperty = FindFProperty<FProperty>(UText3DCharacterBase::StaticClass(), VisibilityPropertyName);
 	check(VisibilityProperty);
 
-	OutProperties.Add(FPropertyAnimatorCoreData(TextRootComponent, VisibilityProperty, nullptr, UPropertyAnimatorTextResolver::StaticClass()));
-}
\ No newline at end of file
+	OutProperties.Add(FPropertyAnimatorCoreData(TextComponent, VisibilityProperty, nullptr, UPropertyAnimatorTextResolver::StaticClass()));
+}
diff --git a/Engine/Plugins/Experimental/PropertyAnimator/Source/PropertyAnimator/Private/Properties/PropertyAnimatorTextResolver.cpp b/Engine/Plugins/Experimental/PropertyAnimator/Source/PropertyAnimator/Private/Properties/PropertyAnimatorTextResolver.cpp
--- a/Engine/Plugins/Experimental/PropertyAnimator/Source/PropertyAnimator/Private/Properties/PropertyAnimatorTextResolver.cpp
+++ b/Engine/Plugins/Experimental/PropertyAnimator/Source/PropertyAnimator/Private/Properties/PropertyAnimatorTextResolver.cpp
@@ -12,49 +12,34 @@ void UPropertyAnimatorTextResolver::SetUnit(EPropertyAnimatorTextResolverRangeUn
 	Unit = InUnit;
 }
 
-void UPropertyAnimatorTextResolver::SetStart(float InRangeStart)
+void UPropertyAnimatorTextResolver::SetPercentageRange(float InPercentageRange)
 {
-	Start = FMath::Clamp(InRangeStart, 0, 100);
+	PercentageRange = FMath::Clamp(InPercentageRange, 0.f, 100.f);
 }
 
-void UPropertyAnimatorTextResolver::SetEnd(float InRangeEnd)
+void UPropertyAnimatorTextResolver::SetPercentageOffset(float InPercentageOffset)
 {
-	End = FMath::Clamp(InRangeEnd, 0, 100);
+	PercentageOffset = InPercentageOffset;
 }
 
-void UPropertyAnimatorTextResolver::SetOffset(float InRangeOffset)
+void UPropertyAnimatorTextResolver::SetCharacterRangeCount(int32 InCharacterRangeCount)
 {
-	Offset = InRangeOffset;
+	CharacterRangeCount = FMath::Max(InCharacterRangeCount, 0);
 }
 
-void UPropertyAnimatorTextResolver::SetCharacterStartIndex(int32 InRangeStart)
+void UPropertyAnimatorTextResolver::SetCharacterOffsetCount(int32 InCharacterOffsetCount)
 {
-	CharacterStartIndex = FMath::Max(InRangeStart, 0);
+	CharacterOffsetCount = InCharacterOffsetCount;
 }
 
-void UPropertyAnimatorTextResolver::SetCharacterEndIndex(int32 InRangeEnd)
+void UPropertyAnimatorTextResolver::SetWordRangeCount(int32 InWordRangeCount)
 {
-	CharacterEndIndex = FMath::Max(InRangeEnd, 0);
+	WordRangeCount = FMath::Max(InWordRangeCount, 0);
 }
 
-void UPropertyAnimatorTextResolver::SetCharacterOffsetIndex(int32 InRangeOffset)
+void UPropertyAnimatorTextResolver::SetWordOffsetCount(int32 InWordOffsetCount)
 {
-	CharacterOffsetIndex = InRangeOffset;
-}
-
-void UPropertyAnimatorTextResolver::SetWordStartIndex(int32 InRangeStart)
-{
-	WordStartIndex = FMath::Max(InRangeStart, 0);
-}
-
-void UPropertyAnimatorTextResolver::SetWordEndIndex(int32 InRangeEnd)
-{
-	WordEndIndex = FMath::Max(InRangeEnd, 0);
-}
-
-void UPropertyAnimatorTextResolver::SetWordOffsetIndex(int32 InRangeOffset)
-{
-	WordOffsetIndex = InRangeOffset;
+	WordOffsetCount = InWordOffsetCount;
 }
 
 void UPropertyAnimatorTextResolver::SetDirection(EPropertyAnimatorTextResolverRangeDirection InDirection)
@@ -103,6 +88,7 @@ void UPropertyAnimatorTextResolver::GetTemplateProperties(UObject* InContext, TS
 	}
 
 	TArray<FPropertyAnimatorCoreData> TemplateProperties;
+	TemplateProperties.Reserve(4);
 
 	FProperty* RelativeLocation = FindFProperty<FProperty>(UText3DCharacterBase::StaticClass(), UText3DCharacterBase::GetRelativeLocationPropertyName());
 	FPropertyAnimatorCoreData RelativeLocationProperty(TextComponent, RelativeLocation, nullptr, GetClass());
@@ -116,6 +102,10 @@ void UPropertyAnimatorTextResolver::GetTemplateProperties(UObject* InContext, TS
 	FPropertyAnimatorCoreData RelativeScaleProperty(TextComponent, RelativeScale, nullptr, GetClass());
 	TemplateProperties.Add(RelativeScaleProperty);
 
+	FProperty* Visible = FindFProperty<FProperty>(UText3DCharacterBase::StaticClass(), UText3DCharacterBase::GetVisiblePropertyName());
+	FPropertyAnimatorCoreData VisibleProperty(TextComponent, Visible, nullptr, GetClass());
+	TemplateProperties.Add(VisibleProperty);
+
 	if (!InSearchPath || InSearchPath->IsEmpty())
 	{
 		OutProperties.Append(TemplateProperties);
@@ -155,7 +145,7 @@ void UPropertyAnimatorTextResolver::GetTemplateProperties(UObject* InContext, TS
 	}
 }
 
-void UPropertyAnimatorTextResolver::ResolveProperties(const FPropertyAnimatorCoreData& InTemplateProperty, TArray<FPropertyAnimatorCoreData>& OutProperties, bool bInForEvaluation)
+void UPropertyAnimatorTextResolver::ResolveTemplateProperties(const FPropertyAnimatorCoreData& InTemplateProperty, TArray<FPropertyAnimatorCoreData>& OutProperties, bool bInForEvaluation)
 {
 	if (!InTemplateProperty.IsResolvable())
 	{
@@ -171,18 +161,11 @@ void UPropertyAnimatorTextResolver::ResolveProperties(const FPropertyAnimatorCor
 	const TArray<FProperty*> ChainProperties = InTemplateProperty.GetChainProperties();
 
 	// Gather each character in the text
-	for (int32 CharacterIndex = 0; CharacterIndex < TextComponent->GetCharacterCount(); CharacterIndex++)
+	TextComponent->ForEachCharacter([this, &OutProperties, &ChainProperties](UText3DCharacterBase* InCharacter, uint16 InIndex, uint16 InCount)
 	{
-		UText3DCharacterBase* Character = TextComponent->GetCharacter(CharacterIndex);
-
-		if (!Character)
-		{
-			continue;
-		}
-
-		FPropertyAnimatorCoreData CharacterProperty(Character, ChainProperties);
-		OutProperties.Add(CharacterProperty);
-	}
+		FPropertyAnimatorCoreData CharacterProperty(InCharacter, ChainProperties);
+		OutProperties.Emplace(CharacterProperty);
+	});
 
 	if (!bInForEvaluation || OutProperties.IsEmpty())
 	{
@@ -197,39 +180,59 @@ void UPropertyAnimatorTextResolver::ResolveProperties(const FPropertyAnimatorCor
 	{
 		case EPropertyAnimatorTextResolverRangeUnit::Percentage:
 		{
-			float StartPercentage = Start / 100.f;
-			float EndPercentage = End / 100.f;
-			float OffsetPercentage = Offset / 100.f;
+			const float Range = PercentageRange / 100.f;
+			const float Offset = PercentageOffset / 100.f;
+
+			const int32 RangeCount = FMath::Clamp(FMath::RoundToInt(MaxIndex * Range), 0, MaxIndex);
+			int32 RangeOffset = FMath::RoundToInt(MaxIndex * Offset);
+
+			int32 CharacterStart = 0;
+			int32 CharacterEnd = RangeCount;
 
 			if (Direction == EPropertyAnimatorTextResolverRangeDirection::RightToLeft)
 			{
-				const float Temp = 1.f - StartPercentage;
-				StartPercentage = 1.f - EndPercentage;
-				EndPercentage = Temp;
-				OffsetPercentage *= -1;
+				CharacterStart = MaxIndex - RangeCount;
+				CharacterEnd = MaxIndex;
+				RangeOffset *= -1;
+			}
+			else if (Direction == EPropertyAnimatorTextResolverRangeDirection::FromCenter)
+			{
+				const int32 CharacterMid = MaxIndex / 2;
+				const int32 ExpansionLeft = RangeCount / 2;
+				const int32 ExpansionRight = RangeCount - ExpansionLeft;
+				CharacterStart = CharacterMid - ExpansionLeft;
+				CharacterEnd = CharacterMid + ExpansionRight;
 			}
 
-			BeginIndex = StartPercentage * MaxIndex + OffsetPercentage * MaxIndex;
-			EndIndex = EndPercentage * MaxIndex + OffsetPercentage * MaxIndex;
+			BeginIndex = CharacterStart + RangeOffset;
+			EndIndex = CharacterEnd + RangeOffset;
 		}
 		break;
 
 		case EPropertyAnimatorTextResolverRangeUnit::Character:
 		{
-			int32 CharacterStart = CharacterStartIndex;
-			int32 CharacterEnd = CharacterEndIndex;
-			int32 CharacterOffset = CharacterOffsetIndex;
+			const int32 RangeCount = CharacterRangeCount;
+			int32 RangeOffset = CharacterOffsetCount;
+			int32 CharacterStart = 0;
+			int32 CharacterEnd = CharacterRangeCount - 1;
 
 			if (Direction == EPropertyAnimatorTextResolverRangeDirection::RightToLeft)
 			{
-				const int32 Temp = MaxIndex - CharacterStart;
-				CharacterStart = MaxIndex - CharacterEnd;
-				CharacterEnd = Temp;
-				CharacterOffset *= -1;
+				CharacterStart = MaxIndex - RangeCount;
+				CharacterEnd = MaxIndex;
+				RangeOffset *= -1;
+			}
+			else if (Direction == EPropertyAnimatorTextResolverRangeDirection::FromCenter)
+			{
+				const int32 CharacterMid = MaxIndex / 2;
+				const int32 ExpansionLeft = RangeCount / 2;
+				const int32 ExpansionRight = RangeCount - ExpansionLeft;
+				CharacterStart = CharacterMid - ExpansionLeft;
+				CharacterEnd = CharacterMid + ExpansionRight;
 			}
 
-			BeginIndex = CharacterStart + CharacterOffset;
-			EndIndex = CharacterEnd + CharacterOffset;
+			BeginIndex = CharacterStart + RangeOffset;
+			EndIndex = CharacterEnd + RangeOffset;
 		}
 		break;
 
@@ -243,36 +246,43 @@ void UPropertyAnimatorTextResolver::ResolveProperties(const FPropertyAnimatorCor
 			}
 
 			const int32 WordCount = TextStats.Words.Num();
-			int32 WordStart = WordStartIndex;
-			int32 WordEnd = WordEndIndex;
-			int32 WordOffset = WordOffsetIndex;
+			const int32 RangeCount = FMath::Clamp(WordRangeCount, 0, WordCount);
+			int32 WordOffset = WordOffsetCount;
+
+			int32 WordStart = 0;
+			int32 WordEnd = RangeCount - 1;
 
 			if (Direction == EPropertyAnimatorTextResolverRangeDirection::RightToLeft)
 			{
-				const int32 Temp = WordCount - WordStart;
-				WordStart = WordCount - WordEnd;
-				WordEnd = Temp;
+				WordStart = WordCount - RangeCount;
+				WordEnd = WordCount - 1;
 				WordOffset *= -1;
 			}
+			else if (Direction == EPropertyAnimatorTextResolverRangeDirection::FromCenter)
+			{
+				const int32 WordMid = WordCount / 2;
+				const int32 ExpansionLeft = RangeCount / 2;
+				const int32 ExpansionRight = RangeCount - ExpansionLeft;
+
+				WordStart = WordMid - ExpansionLeft;
+				WordEnd = WordMid + ExpansionRight - 1;
+			}
+
+			WordStart += WordOffset;
+			WordEnd += WordOffset;
+
+			if (TextStats.Words.IsValidIndex(WordStart))
+			{
+				BeginIndex = TextStats.Words[WordStart].RenderRange.BeginIndex;
+			}
 
-			if (WordStart != WordEnd)
+			if (TextStats.Words.IsValidIndex(WordEnd))
+			{
+				EndIndex = TextStats.Words[WordEnd].RenderRange.EndIndex;
+			}
+			else if (WordEnd >= WordCount && WordStart < WordCount)
 			{
-				WordStart += WordOffset;
-				WordEnd += WordOffset - 1;
-
-				if (TextStats.Words.IsValidIndex(WordStart))
-				{
-					BeginIndex = TextStats.Words[WordStart].RenderRange.BeginIndex;
-				}
-
-				if (TextStats.Words.IsValidIndex(WordEnd))
-				{
-					EndIndex = TextStats.Words[WordEnd].RenderRange.EndIndex;
-				}
-				else if (WordEnd >= WordCount && WordStart < WordCount)
-				{
-					EndIndex = TextStats.Words.Last().RenderRange.EndIndex;
-				}
+				EndIndex = TextStats.Words.Last().RenderRange.EndIndex;
 			}
 		}
 		break;
@@ -307,45 +317,47 @@ bool UPropertyAnimatorTextResolver::ImportPreset(const UPropertyAnimatorCorePres
 		ResolverArchive->Get(GET_MEMBER_NAME_STRING_CHECKED(UPropertyAnimatorTextResolver, Unit), UnitValue);
 		SetUnit(static_cast<EPropertyAnimatorTextResolverRangeUnit>(UnitValue));
 
-		double StartValue = Start;
-		ResolverArchive->Get(GET_MEMBER_NAME_STRING_CHECKED(UPropertyAnimatorTextResolver, Start), StartValue);
-		SetStart(StartValue);
-
-		double EndValue = End;
-		ResolverArchive->Get(GET_MEMBER_NAME_STRING_CHECKED(UPropertyAnimatorTextResolver, End), EndValue);
-		SetEnd(EndValue);
-
-		double OffsetValue = Offset;
-		ResolverArchive->Get(GET_MEMBER_NAME_STRING_CHECKED(UPropertyAnimatorTextResolver, Offset), OffsetValue);
-		SetOffset(OffsetValue);
-
-		int64 CharStartIndexValue = CharacterStartIndex;
-		ResolverArchive->Get(GET_MEMBER_NAME_STRING_CHECKED(UPropertyAnimatorTextResolver, CharacterStartIndex), CharStartIndexValue);
-		SetCharacterStartIndex(CharStartIndexValue);
+		{
+			double PercentageRangeValue = PercentageRange;
+			ResolverArchive->Get(GET_MEMBER_NAME_STRING_CHECKED(UPropertyAnimatorTextResolver, PercentageRange), PercentageRangeValue);
+			SetPercentageRange(PercentageRangeValue);
+		}
 
-		int64 CharEndIndexValue = CharacterEndIndex;
-		ResolverArchive->Get(GET_MEMBER_NAME_STRING_CHECKED(UPropertyAnimatorTextResolver, CharacterEndIndex), CharEndIndexValue);
-		SetCharacterEndIndex(CharEndIndexValue);
+		{
+			double PercentageOffsetValue = PercentageOffset;
+			ResolverArchive->Get(GET_MEMBER_NAME_STRING_CHECKED(UPropertyAnimatorTextResolver, PercentageOffset), PercentageOffsetValue);
+			SetPercentageOffset(PercentageOffsetValue);
+		}
 
-		int64 CharOffsetIndexValue = CharacterOffsetIndex;
-		ResolverArchive->Get(GET_MEMBER_NAME_STRING_CHECKED(UPropertyAnimatorTextResolver, CharacterOffsetIndex), CharOffsetIndexValue);
-		SetCharacterOffsetIndex(CharOffsetIndexValue);
+		{
+			int64 CharacterRangeCountValue = CharacterRangeCount;
+			ResolverArchive->Get(GET_MEMBER_NAME_STRING_CHECKED(UPropertyAnimatorTextResolver, CharacterRangeCount), CharacterRangeCountValue);
+			SetCharacterRangeCount(CharacterRangeCountValue);
+		}
 
-		int64 WordStartIndexValue = WordStartIndex;
-		ResolverArchive->Get(GET_MEMBER_NAME_STRING_CHECKED(UPropertyAnimatorTextResolver, WordStartIndex), WordStartIndexValue);
-		SetWordStartIndex(WordStartIndexValue);
+		{
+			int64 CharacterOffsetCountValue = CharacterOffsetCount;
+			ResolverArchive->Get(GET_MEMBER_NAME_STRING_CHECKED(UPropertyAnimatorTextResolver, CharacterOffsetCount), CharacterOffsetCountValue);
+			SetCharacterOffsetCount(CharacterOffsetCountValue);
+		}
 
-		int64 WordEndIndexValue = WordEndIndex;
-		ResolverArchive->Get(GET_MEMBER_NAME_STRING_CHECKED(UPropertyAnimatorTextResolver, WordEndIndex), WordEndIndexValue);
-		SetWordEndIndex(WordEndIndexValue);
+		{
+			int64 WordRangeCountValue = WordRangeCount;
+			ResolverArchive->Get(GET_MEMBER_NAME_STRING_CHECKED(UPropertyAnimatorTextResolver, WordRangeCount), WordRangeCountValue);
+			SetWordRangeCount(WordRangeCountValue);
+		}
 
-		int64 WordOffsetIndexValue = WordOffsetIndex;
-		ResolverArchive->Get(GET_MEMBER_NAME_STRING_CHECKED(UPropertyAnimatorTextResolver, WordOffsetIndex), WordOffsetIndexValue);
-		SetWordOffsetIndex(WordOffsetIndexValue);
+		{
+			int64 WordOffsetCountValue = WordOffsetCount;
+			ResolverArchive->Get(GET_MEMBER_NAME_STRING_CHECKED(UPropertyAnimatorTextResolver, WordOffsetCount), WordOffsetCountValue);
+			SetWordOffsetCount(WordOffsetCountValue);
+		}
 
-		uint64 DirectionValue = static_cast<uint64>(Direction);
-		ResolverArchive->Get(GET_MEMBER_NAME_STRING_CHECKED(UPropertyAnimatorTextResolver, Direction), DirectionValue);
-		SetDirection(static_cast<EPropertyAnimatorTextResolverRangeDirection>(WordOffsetIndexValue));
+		{
+			uint64 DirectionValue = static_cast<uint64>(Direction);
+			ResolverArchive->Get(GET_MEMBER_NAME_STRING_CHECKED(UPropertyAnimatorTextResolver, Direction), DirectionValue);
+			SetDirection(static_cast<EPropertyAnimatorTextResolverRangeDirection>(DirectionValue));
+		}
 
 		return true;
 	}
@@ -360,15 +372,12 @@ bool UPropertyAnimatorTextResolver::ExportPreset(const UPropertyAnimatorCorePres
 		const TSharedPtr<FPropertyAnimatorCorePresetObjectArchive> ResolverArchive = OutValue->AsMutableObject();
 
 		ResolverArchive->Set(GET_MEMBER_NAME_STRING_CHECKED(UPropertyAnimatorTextResolver, Unit), static_cast<uint64>(Unit));
-		ResolverArchive->Set(GET_MEMBER_NAME_STRING_CHECKED(UPropertyAnimatorTextResolver, Start), Start);
-		ResolverArchive->Set(GET_MEMBER_NAME_STRING_CHECKED(UPropertyAnimatorTextResolver, End), End);
-		ResolverArchive->Set(GET_MEMBER_NAME_STRING_CHECKED(UPropertyAnimatorTextResolver, Offset), Offset);
-		ResolverArchive->Set(GET_MEMBER_NAME_STRING_CHECKED(UPropertyAnimatorTextResolver, CharacterStartIndex), static_cast<int64>(CharacterStartIndex));
-		ResolverArchive->Set(GET_MEMBER_NAME_STRING_CHECKED(UPropertyAnimatorTextResolver, CharacterEndIndex), static_cast<int64>(CharacterEndIndex));
-		ResolverArchive->Set(GET_MEMBER_NAME_STRING_CHECKED(UPropertyAnimatorTextResolver, CharacterOffsetIndex), static_cast<int64>(CharacterOffsetIndex));
-		ResolverArchive->Set(GET_MEMBER_NAME_STRING_CHECKED(UPropertyAnimatorTextResolver, WordStartIndex), static_cast<int64>(WordStartIndex));
-		ResolverArchive->Set(GET_MEMBER_NAME_STRING_CHECKED(UPropertyAnimatorTextResolver, WordEndIndex), static_cast<int64>(WordEndIndex));
-		ResolverArchive->Set(GET_MEMBER_NAME_STRING_CHECKED(UPropertyAnimatorTextResolver, WordOffsetIndex), static_cast<int64>(WordOffsetIndex));
+		ResolverArchive->Set(GET_MEMBER_NAME_STRING_CHECKED(UPropertyAnimatorTextResolver, PercentageRange), PercentageRange);
+		ResolverArchive->Set(GET_MEMBER_NAME_STRING_CHECKED(UPropertyAnimatorTextResolver, PercentageOffset), PercentageOffset);
+		ResolverArchive->Set(GET_MEMBER_NAME_STRING_CHECKED(UPropertyAnimatorTextResolver, CharacterRangeCount), static_cast<int64>(CharacterRangeCount));
+		ResolverArchive->Set(GET_MEMBER_NAME_STRING_CHECKED(UPropertyAnimatorTextResolver, CharacterOffsetCount), static_cast<int64>(CharacterOffsetCount));
+		ResolverArchive->Set(GET_MEMBER_NAME_STRING_CHECKED(UPropertyAnimatorTextResolver, WordRangeCount), static_cast<int64>(WordRangeCount));
+		ResolverArchive->Set(GET_MEMBER_NAME_STRING_CHECKED(UPropertyAnimatorTextResolver, WordOffsetCount), static_cast<int64>(WordOffsetCount));
 		ResolverArchive->Set(GET_MEMBER_NAME_STRING_CHECKED(UPropertyAnimatorTextResolver, Direction), static_cast<uint64>(Direction));
 
 		return true;
diff --git a/Engine/Plugins/Experimental/PropertyAnimator/Source/PropertyAnimator/Private/Properties/PropertyAnimatorTextResolver.h b/Engine/Plugins/Experimental/PropertyAnimator/Source/PropertyAnimator/Private/Properties/PropertyAnimatorTextResolver.h
--- a/Engine/Plugins/Experimental/PropertyAnimator/Source/PropertyAnimator/Private/Properties/PropertyAnimatorTextResolver.h
+++ b/Engine/Plugins/Experimental/PropertyAnimator/Source/PropertyAnimator/Private/Properties/PropertyAnimatorTextResolver.h
@@ -18,7 +18,8 @@ UENUM()
 enum class EPropertyAnimatorTextResolverRangeDirection : uint8
 {
 	LeftToRight,
-	RightToLeft
+	RightToLeft,
+	FromCenter
 };
 
 /**
@@ -43,104 +44,83 @@ public:
 		return Unit;
 	}
 
-	PROPERTYANIMATOR_API void SetStart(float InRangeStart);
-	float GetStart() const
+	float GetPercentageRange() const
 	{
-		return Start;
+		return PercentageRange;
 	}
 
-	PROPERTYANIMATOR_API void SetEnd(float InRangeEnd);
-	float GetEnd() const
-	{
-		return End;
-	}
+	void SetPercentageRange(float InPercentageRange);
 
-	PROPERTYANIMATOR_API void SetOffset(float InRangeOffset);
-	float GetOffset() const
+	float GetPercentageOffset() const
 	{
-		return Offset;
+		return PercentageOffset;
 	}
 
-	PROPERTYANIMATOR_API void SetCharacterStartIndex(int32 InRangeStart);
-	int32 GetCharacterStartIndex() const
-	{
-		return CharacterStartIndex;
-	}
+	void SetPercentageOffset(float InPercentageOffset);
 
-	PROPERTYANIMATOR_API void SetCharacterEndIndex(int32 InRangeEnd);
-	int32 GetCharacterEndIndex() const
+	int32 GetCharacterRangeCount() const
 	{
-		return CharacterEndIndex;
+		return CharacterRangeCount;
 	}
 
-	PROPERTYANIMATOR_API void SetCharacterOffsetIndex(int32 InRangeOffset);
-	int32 GetCharacterOffsetIndex() const
-	{
-		return CharacterOffsetIndex;
-	}
+	void SetCharacterRangeCount(int32 InCharacterRangeCount);
 
-	PROPERTYANIMATOR_API void SetWordStartIndex(int32 InRangeStart);
-	int32 GetWordStartIndex() const
+	int32 GetCharacterOffsetCount() const
 	{
-		return WordStartIndex;
+		return CharacterOffsetCount;
 	}
 
-	PROPERTYANIMATOR_API void SetWordEndIndex(int32 InRangeEnd);
-	int32 GetWordEndIndex() const
+	void SetCharacterOffsetCount(int32 InCharacterOffsetCount);
+
+	int32 GetWordRangeCount() const
 	{
-		return WordEndIndex;
+		return WordRangeCount;
 	}
 
-	PROPERTYANIMATOR_API void SetWordOffsetIndex(int32 InRangeOffset);
-	int32 GetWordOffsetIndex() const
+	void SetWordRangeCount(int32 InWordRangeCount);
+
+	int32 GetWordOffsetCount() const
 	{
-		return WordOffsetIndex;
+		return WordOffsetCount;
 	}
 
+	void SetWordOffsetCount(int32 InWordOffsetCount);
+ 
 	PROPERTYANIMATOR_API void SetDirection(EPropertyAnimatorTextResolverRangeDirection InDirection);
 	EPropertyAnimatorTextResolverRangeDirection GetDirection() const
 	{
 		return Direction;
 	}
 
+protected:
 	//~ Begin UPropertyAnimatorCoreResolver
 	virtual bool FixUpProperty(FPropertyAnimatorCoreData& InOldProperty) override;
 	virtual void GetTemplateProperties(UObject* InContext, TSet<FPropertyAnimatorCoreData>& OutProperties, const TArray<FName>* InSearchPath) override;
-	virtual void ResolveProperties(const FPropertyAnimatorCoreData& InTemplateProperty, TArray<FPropertyAnimatorCoreData>& OutProperties, bool bInForEvaluation) override;
+	virtual void ResolveTemplateProperties(const FPropertyAnimatorCoreData& InTemplateProperty, TArray<FPropertyAnimatorCoreData>& OutProperties, bool bInForEvaluation) override;
 	virtual bool ImportPreset(const UPropertyAnimatorCorePresetBase* InPreset, const TSharedRef<FPropertyAnimatorCorePresetArchive>& InValue) override;
 	virtual bool ExportPreset(const UPropertyAnimatorCorePresetBase* InPreset, TSharedPtr<FPropertyAnimatorCorePresetArchive>& OutValue) const override;
 	//~ End UPropertyAnimatorCoreResolver
 
-protected:
 	UPROPERTY(EditInstanceOnly, Setter, Getter, Category="Animator")
 	EPropertyAnimatorTextResolverRangeUnit Unit = EPropertyAnimatorTextResolverRangeUnit::Percentage;
 
-	UPROPERTY(EditInstanceOnly, Setter, Getter, Category="Animator", meta=(ClampMin="0", ClampMax="100", Units=Percent, EditCondition="Unit == EPropertyAnimatorTextResolverRangeUnit::Percentage", EditConditionHides))
-	float Start = 0.f;
-
-	UPROPERTY(EditInstanceOnly, Setter, Getter, Category="Animator", meta=(ClampMin="0", ClampMax="100", Units=Percent, EditCondition="Unit == EPropertyAnimatorTextResolverRangeUnit::Percentage", EditConditionHides))
-	float End = 100.f;
-
-	UPROPERTY(EditInstanceOnly, Setter, Getter, Category="Animator", meta=(Units=Percent, Delta="1", EditCondition="Unit == EPropertyAnimatorTextResolverRangeUnit::Percentage", EditConditionHides))
-	float Offset = 0.f;
-
-	UPROPERTY(EditInstanceOnly, Setter, Getter, Category="Animator", meta=(ClampMin="0", EditCondition="Unit == EPropertyAnimatorTextResolverRangeUnit::Character", EditConditionHides))
-	int32 CharacterStartIndex = 0;
+	UPROPERTY(EditInstanceOnly, Setter, Getter, Category="Animator", meta=(Delta="1", ClampMin="0", ClampMax="100", Units=Percent, EditCondition="Unit == EPropertyAnimatorTextResolverRangeUnit::Percentage", EditConditionHides))
+	float PercentageRange = 100.f;
 
-	UPROPERTY(EditInstanceOnly, Setter, Getter, Category="Animator", meta=(ClampMin="0", EditCondition="Unit == EPropertyAnimatorTextResolverRangeUnit::Character", EditConditionHides))
-	int32 CharacterEndIndex = 100;
+	UPROPERTY(EditInstanceOnly, Setter, Getter, Category="Animator", meta=(Delta="1", Units=Percent, EditCondition="Unit == EPropertyAnimatorTextResolverRangeUnit::Percentage", EditConditionHides))
+	float PercentageOffset = 0.f;
 
-	UPROPERTY(EditInstanceOnly, Setter, Getter, Category="Animator", meta=(EditCondition="Unit == EPropertyAnimatorTextResolverRangeUnit::Character", EditConditionHides))
-	int32 CharacterOffsetIndex = 0;
+	UPROPERTY(EditInstanceOnly, Setter, Getter, Category="Animator", meta=(Delta="1", ClampMin="0", EditCondition="Unit == EPropertyAnimatorTextResolverRangeUnit::Character", EditConditionHides))
+	int32 CharacterRangeCount = 100;
 
-	UPROPERTY(EditInstanceOnly, Setter, Getter, Category="Animator", meta=(ClampMin="0", EditCondition="Unit == EPropertyAnimatorTextResolverRangeUnit::Word", EditConditionHides))
-	int32 WordStartIndex = 0;
+	UPROPERTY(EditInstanceOnly, Setter, Getter, Category="Animator", meta=(Delta="1", EditCondition="Unit == EPropertyAnimatorTextResolverRangeUnit::Character", EditConditionHides))
+	int32 CharacterOffsetCount = 0;
 
-	UPROPERTY(EditInstanceOnly, Setter, Getter, Category="Animator", meta=(ClampMin="0", EditCondition="Unit == EPropertyAnimatorTextResolverRangeUnit::Word", EditConditionHides))
-	int32 WordEndIndex = 100;
+	UPROPERTY(EditInstanceOnly, Setter, Getter, Category="Animator", meta=(Delta="1", ClampMin="0", EditCondition="Unit == EPropertyAnimatorTextResolverRangeUnit::Word", EditConditionHides))
+	int32 WordRangeCount = 100;
 
-	UPROPERTY(EditInstanceOnly, Setter, Getter, Category="Animator", meta=(EditCondition="Unit == EPropertyAnimatorTextResolverRangeUnit::Word", EditConditionHides))
-	int32 WordOffsetIndex = 0;
+	UPROPERTY(EditInstanceOnly, Setter, Getter, Category="Animator", meta=(Delta="1", EditCondition="Unit == EPropertyAnimatorTextResolverRangeUnit::Word", EditConditionHides))
+	int32 WordOffsetCount = 0;
 
 	UPROPERTY(EditInstanceOnly, Setter, Getter, Category="Animator")
 	EPropertyAnimatorTextResolverRangeDirection Direction = EPropertyAnimatorTextResolverRangeDirection::LeftToRight;
diff --git a/Engine/Plugins/Experimental/PropertyAnimatorCore/Source/PropertyAnimatorCore/Private/Properties/Converters/PropertyAnimatorCoreFloatToBoolConverter.cpp b/Engine/Plugins/Experimental/PropertyAnimatorCore/Source/PropertyAnimatorCore/Private/Properties/Converters/PropertyAnimatorCoreFloatToBoolConverter.cpp
--- a/Engine/Plugins/Experimental/PropertyAnimatorCore/Source/PropertyAnimatorCore/Private/Properties/Converters/PropertyAnimatorCoreFloatToBoolConverter.cpp
+++ b/Engine/Plugins/Experimental/PropertyAnimatorCore/Source/PropertyAnimatorCore/Private/Properties/Converters/PropertyAnimatorCoreFloatToBoolConverter.cpp
@@ -15,7 +15,7 @@ bool UPropertyAnimatorCoreFloatToBoolConverter::IsConversionSupported(const FPro
 	return Super::IsConversionSupported(InFromProperty, InToProperty);
 }
 
-bool UPropertyAnimatorCoreFloatToBoolConverter::Convert(const FPropertyBagPropertyDesc& InFromProperty, const FInstancedPropertyBag& InFromBag, const FPropertyBagPropertyDesc& InToProperty, FInstancedPropertyBag& InToBag, const FInstancedStruct* InRule)
+bool UPropertyAnimatorCoreFloatToBoolConverter::Convert(const FPropertyBagPropertyDesc& InFromProperty, const FInstancedPropertyBag& InFromBag, const FPropertyBagPropertyDesc& InToProperty, FInstancedPropertyBag& InToBag, const TInstancedStruct<FPropertyAnimatorCoreConverterRuleBase>& InRule)
 {
 	TValueOrError<float, EPropertyBagResult> FromValueResult = InFromBag.GetValueFloat(InFromProperty.Name);
 
@@ -25,7 +25,7 @@ bool UPropertyAnimatorCoreFloatToBoolConverter::Convert(const FPropertyBagProper
 	}
 
 	bool bOutValue;
-	if (TValueConverterTraits<float, bool, FBoolConverterRule>::Convert(FromValueResult.GetValue(), bOutValue, &InRule->Get<FBoolConverterRule>()))
+	if (TValueConverterTraits<float, bool, FBoolConverterRule>::Convert(FromValueResult.GetValue(), bOutValue, InRule.GetPtr<FBoolConverterRule>()))
 	{
 		InToBag.SetValueBool(InToProperty.Name, bOutValue);
 		return true;
diff --git a/Engine/Plugins/Experimental/PropertyAnimatorCore/Source/PropertyAnimatorCore/Private/Properties/Converters/PropertyAnimatorCoreFloatToBoolConverter.h b/Engine/Plugins/Experimental/PropertyAnimatorCore/Source/PropertyAnimatorCore/Private/Properties/Converters/PropertyAnimatorCoreFloatToBoolConverter.h
--- a/Engine/Plugins/Experimental/PropertyAnimatorCore/Source/PropertyAnimatorCore/Private/Properties/Converters/PropertyAnimatorCoreFloatToBoolConverter.h
+++ b/Engine/Plugins/Experimental/PropertyAnimatorCore/Source/PropertyAnimatorCore/Private/Properties/Converters/PropertyAnimatorCoreFloatToBoolConverter.h
@@ -13,7 +13,7 @@ class UPropertyAnimatorCoreFloatToBoolConverter : public UPropertyAnimatorCoreCo
 public:
 	//~ Begin UPropertyAnimatorCoreConverterBase
 	virtual bool IsConversionSupported(const FPropertyBagPropertyDesc& InFromProperty, const FPropertyBagPropertyDesc& InToProperty) const override;
-	virtual bool Convert(const FPropertyBagPropertyDesc& InFromProperty, const FInstancedPropertyBag& InFromBag, const FPropertyBagPropertyDesc& InToProperty, FInstancedPropertyBag& InToBag, const FInstancedStruct* InRule) override;
+	virtual bool Convert(const FPropertyBagPropertyDesc& InFromProperty, const FInstancedPropertyBag& InFromBag, const FPropertyBagPropertyDesc& InToProperty, FInstancedPropertyBag& InToBag, const TInstancedStruct<FPropertyAnimatorCoreConverterRuleBase>& InRule) override;
 	virtual UScriptStruct* GetConversionRuleStruct() const override;
 	//~ End UPropertyAnimatorCoreConverterBase
 };
\ No newline at end of file
diff --git a/Engine/Plugins/Experimental/PropertyAnimatorCore/Source/PropertyAnimatorCore/Private/Properties/Converters/PropertyAnimatorCoreFloatToDoubleConverter.cpp b/Engine/Plugins/Experimental/PropertyAnimatorCore/Source/PropertyAnimatorCore/Private/Properties/Converters/PropertyAnimatorCoreFloatToDoubleConverter.cpp
--- a/Engine/Plugins/Experimental/PropertyAnimatorCore/Source/PropertyAnimatorCore/Private/Properties/Converters/PropertyAnimatorCoreFloatToDoubleConverter.cpp
+++ b/Engine/Plugins/Experimental/PropertyAnimatorCore/Source/PropertyAnimatorCore/Private/Properties/Converters/PropertyAnimatorCoreFloatToDoubleConverter.cpp
@@ -15,7 +15,7 @@ bool UPropertyAnimatorCoreFloatToDoubleConverter::IsConversionSupported(const FP
 	return Super::IsConversionSupported(InFromProperty, InToProperty);
 }
 
-bool UPropertyAnimatorCoreFloatToDoubleConverter::Convert(const FPropertyBagPropertyDesc& InFromProperty, const FInstancedPropertyBag& InFromBag, const FPropertyBagPropertyDesc& InToProperty, FInstancedPropertyBag& InToBag, const FInstancedStruct* InRule)
+bool UPropertyAnimatorCoreFloatToDoubleConverter::Convert(const FPropertyBagPropertyDesc& InFromProperty, const FInstancedPropertyBag& InFromBag, const FPropertyBagPropertyDesc& InToProperty, FInstancedPropertyBag& InToBag, const TInstancedStruct<FPropertyAnimatorCoreConverterRuleBase>& InRule)
 {
 	TValueOrError<float, EPropertyBagResult> FromValueResult = InFromBag.GetValueFloat(InFromProperty.Name);
 
diff --git a/Engine/Plugins/Experimental/PropertyAnimatorCore/Source/PropertyAnimatorCore/Private/Properties/Converters/PropertyAnimatorCoreFloatToDoubleConverter.h b/Engine/Plugins/Experimental/PropertyAnimatorCore/Source/PropertyAnimatorCore/Private/Properties/Converters/PropertyAnimatorCoreFloatToDoubleConverter.h
--- a/Engine/Plugins/Experimental/PropertyAnimatorCore/Source/PropertyAnimatorCore/Private/Properties/Converters/PropertyAnimatorCoreFloatToDoubleConverter.h
+++ b/Engine/Plugins/Experimental/PropertyAnimatorCore/Source/PropertyAnimatorCore/Private/Properties/Converters/PropertyAnimatorCoreFloatToDoubleConverter.h
@@ -13,7 +13,7 @@ class UPropertyAnimatorCoreFloatToDoubleConverter : public UPropertyAnimatorCore
 public:
 	//~ Begin UPropertyAnimatorCoreConverterBase
 	virtual bool IsConversionSupported(const FPropertyBagPropertyDesc& InFromProperty, const FPropertyBagPropertyDesc& InToProperty) const override;
-	virtual bool Convert(const FPropertyBagPropertyDesc& InFromProperty, const FInstancedPropertyBag& InFromBag, const FPropertyBagPropertyDesc& InToProperty, FInstancedPropertyBag& InToBag, const FInstancedStruct* InRule) override;
+	virtual bool Convert(const FPropertyBagPropertyDesc& InFromProperty, const FInstancedPropertyBag& InFromBag, const FPropertyBagPropertyDesc& InToProperty, FInstancedPropertyBag& InToBag, const TInstancedStruct<FPropertyAnimatorCoreConverterRuleBase>& InRule) override;
 	virtual UScriptStruct* GetConversionRuleStruct() const override;
 	//~ End UPropertyAnimatorCoreConverterBase
 };
\ No newline at end of file
diff --git a/Engine/Plugins/Experimental/PropertyAnimatorCore/Source/PropertyAnimatorCore/Private/Properties/Converters/PropertyAnimatorCoreFloatToIntConverter.cpp b/Engine/Plugins/Experimental/PropertyAnimatorCore/Source/PropertyAnimatorCore/Private/Properties/Converters/PropertyAnimatorCoreFloatToIntConverter.cpp
--- a/Engine/Plugins/Experimental/PropertyAnimatorCore/Source/PropertyAnimatorCore/Private/Properties/Converters/PropertyAnimatorCoreFloatToIntConverter.cpp
+++ b/Engine/Plugins/Experimental/PropertyAnimatorCore/Source/PropertyAnimatorCore/Private/Properties/Converters/PropertyAnimatorCoreFloatToIntConverter.cpp
@@ -16,7 +16,7 @@ bool UPropertyAnimatorCoreFloatToIntConverter::IsConversionSupported(const FProp
 	return Super::IsConversionSupported(InFromProperty, InToProperty);
 }
 
-bool UPropertyAnimatorCoreFloatToIntConverter::Convert(const FPropertyBagPropertyDesc& InFromProperty, const FInstancedPropertyBag& InFromBag, const FPropertyBagPropertyDesc& InToProperty, FInstancedPropertyBag& InToBag, const FInstancedStruct* InRule)
+bool UPropertyAnimatorCoreFloatToIntConverter::Convert(const FPropertyBagPropertyDesc& InFromProperty, const FInstancedPropertyBag& InFromBag, const FPropertyBagPropertyDesc& InToProperty, FInstancedPropertyBag& InToBag, const TInstancedStruct<FPropertyAnimatorCoreConverterRuleBase>& InRule)
 {
 	TValueOrError<float, EPropertyBagResult> FromValueResult = InFromBag.GetValueFloat(InFromProperty.Name);
 
@@ -26,7 +26,7 @@ bool UPropertyAnimatorCoreFloatToIntConverter::Convert(const FPropertyBagPropert
 	}
 
 	int32 OutValue;
-	if (TValueConverterTraits<float, int32, FInt32ConverterRule>::Convert(FromValueResult.GetValue(), OutValue, &InRule->Get<FInt32ConverterRule>()))
+	if (TValueConverterTraits<float, int32, FInt32ConverterRule>::Convert(FromValueResult.GetValue(), OutValue, InRule.GetPtr<FInt32ConverterRule>()))
 	{
 		InToBag.SetValueDouble(InToProperty.Name, OutValue);
 		return true;
diff --git a/Engine/Plugins/Experimental/PropertyAnimatorCore/Source/PropertyAnimatorCore/Private/Properties/Converters/PropertyAnimatorCoreFloatToIntConverter.h b/Engine/Plugins/Experimental/PropertyAnimatorCore/Source/PropertyAnimatorCore/Private/Properties/Converters/PropertyAnimatorCoreFloatToIntConverter.h
--- a/Engine/Plugins/Experimental/PropertyAnimatorCore/Source/PropertyAnimatorCore/Private/Properties/Converters/PropertyAnimatorCoreFloatToIntConverter.h
+++ b/Engine/Plugins/Experimental/PropertyAnimatorCore/Source/PropertyAnimatorCore/Private/Properties/Converters/PropertyAnimatorCoreFloatToIntConverter.h
@@ -13,7 +13,7 @@ class UPropertyAnimatorCoreFloatToIntConverter : public UPropertyAnimatorCoreCon
 public:
 	//~ Begin UPropertyAnimatorCoreConverterBase
 	virtual bool IsConversionSupported(const FPropertyBagPropertyDesc& InFromProperty, const FPropertyBagPropertyDesc& InToProperty) const override;
-	virtual bool Convert(const FPropertyBagPropertyDesc& InFromProperty, const FInstancedPropertyBag& InFromBag, const FPropertyBagPropertyDesc& InToProperty, FInstancedPropertyBag& InToBag, const FInstancedStruct* InRule) override;
+	virtual bool Convert(const FPropertyBagPropertyDesc& InFromProperty, const FInstancedPropertyBag& InFromBag, const FPropertyBagPropertyDesc& InToProperty, FInstancedPropertyBag& InToBag, const TInstancedStruct<FPropertyAnimatorCoreConverterRuleBase>& InRule) override;
 	virtual UScriptStruct* GetConversionRuleStruct() const override;
 	//~ End UPropertyAnimatorCoreConverterBase
 };
\ No newline at end of file
diff --git a/Engine/Plugins/Experimental/PropertyAnimatorCore/Source/PropertyAnimatorCore/Private/Properties/Converters/PropertyAnimatorCoreStringToTextConverter.cpp b/Engine/Plugins/Experimental/PropertyAnimatorCore/Source/PropertyAnimatorCore/Private/Properties/Converters/PropertyAnimatorCoreStringToTextConverter.cpp
--- a/Engine/Plugins/Experimental/PropertyAnimatorCore/Source/PropertyAnimatorCore/Private/Properties/Converters/PropertyAnimatorCoreStringToTextConverter.cpp
+++ b/Engine/Plugins/Experimental/PropertyAnimatorCore/Source/PropertyAnimatorCore/Private/Properties/Converters/PropertyAnimatorCoreStringToTextConverter.cpp
@@ -15,7 +15,7 @@ bool UPropertyAnimatorCoreStringToTextConverter::IsConversionSupported(const FPr
 	return Super::IsConversionSupported(InFromProperty, InToProperty);
 }
 
-bool UPropertyAnimatorCoreStringToTextConverter::Convert(const FPropertyBagPropertyDesc& InFromProperty, const FInstancedPropertyBag& InFromBag, const FPropertyBagPropertyDesc& InToProperty, FInstancedPropertyBag& InToBag, const FInstancedStruct* InRule)
+bool UPropertyAnimatorCoreStringToTextConverter::Convert(const FPropertyBagPropertyDesc& InFromProperty, const FInstancedPropertyBag& InFromBag, const FPropertyBagPropertyDesc& InToProperty, FInstancedPropertyBag& InToBag, const TInstancedStruct<FPropertyAnimatorCoreConverterRuleBase>& InRule)
 {
 	TValueOrError<FString, EPropertyBagResult> FromValueResult = InFromBag.GetValueString(InFromProperty.Name);
 
diff --git a/Engine/Plugins/Experimental/PropertyAnimatorCore/Source/PropertyAnimatorCore/Private/Properties/Converters/PropertyAnimatorCoreStringToTextConverter.h b/Engine/Plugins/Experimental/PropertyAnimatorCore/Source/PropertyAnimatorCore/Private/Properties/Converters/PropertyAnimatorCoreStringToTextConverter.h
--- a/Engine/Plugins/Experimental/PropertyAnimatorCore/Source/PropertyAnimatorCore/Private/Properties/Converters/PropertyAnimatorCoreStringToTextConverter.h
+++ b/Engine/Plugins/Experimental/PropertyAnimatorCore/Source/PropertyAnimatorCore/Private/Properties/Converters/PropertyAnimatorCoreStringToTextConverter.h
@@ -13,7 +13,7 @@ class UPropertyAnimatorCoreStringToTextConverter : public UPropertyAnimatorCoreC
 public:
 	//~ Begin UPropertyAnimatorCoreConverterBase
 	virtual bool IsConversionSupported(const FPropertyBagPropertyDesc& InFromProperty, const FPropertyBagPropertyDesc& InToProperty) const override;
-	virtual bool Convert(const FPropertyBagPropertyDesc& InFromProperty, const FInstancedPropertyBag& InFromBag, const FPropertyBagPropertyDesc& InToProperty, FInstancedPropertyBag& InToBag, const FInstancedStruct* InRule) override;
+	virtual bool Convert(const FPropertyBagPropertyDesc& InFromProperty, const FInstancedPropertyBag& InFromBag, const FPropertyBagPropertyDesc& InToProperty, FInstancedPropertyBag& InToBag, const TInstancedStruct<FPropertyAnimatorCoreConverterRuleBase>& InRule) override;
 	virtual UScriptStruct* GetConversionRuleStruct() const override;
 	//~ End UPropertyAnimatorCoreConverterBase
 };
\ No newline at end of file
diff --git a/Engine/Plugins/Experimental/PropertyAnimatorCore/Source/PropertyAnimatorCore/Private/Properties/Converters/PropertyAnimatorCoreTextToStringConverter.cpp b/Engine/Plugins/Experimental/PropertyAnimatorCore/Source/PropertyAnimatorCore/Private/Properties/Converters/PropertyAnimatorCoreTextToStringConverter.cpp
--- a/Engine/Plugins/Experimental/PropertyAnimatorCore/Source/PropertyAnimatorCore/Private/Properties/Converters/PropertyAnimatorCoreTextToStringConverter.cpp
+++ b/Engine/Plugins/Experimental/PropertyAnimatorCore/Source/PropertyAnimatorCore/Private/Properties/Converters/PropertyAnimatorCoreTextToStringConverter.cpp
@@ -15,7 +15,7 @@ bool UPropertyAnimatorCoreTextToStringConverter::IsConversionSupported(const FPr
 	return Super::IsConversionSupported(InFromProperty, InToProperty);
 }
 
-bool UPropertyAnimatorCoreTextToStringConverter::Convert(const FPropertyBagPropertyDesc& InFromProperty, const FInstancedPropertyBag& InFromBag, const FPropertyBagPropertyDesc& InToProperty, FInstancedPropertyBag& InToBag, const FInstancedStruct* InRule)
+bool UPropertyAnimatorCoreTextToStringConverter::Convert(const FPropertyBagPropertyDesc& InFromProperty, const FInstancedPropertyBag& InFromBag, const FPropertyBagPropertyDesc& InToProperty, FInstancedPropertyBag& InToBag, const TInstancedStruct<FPropertyAnimatorCoreConverterRuleBase>& InRule)
 {
 	TValueOrError<FText, EPropertyBagResult> FromValueResult = InFromBag.GetValueText(InFromProperty.Name);
 
diff --git a/Engine/Plugins/Experimental/PropertyAnimatorCore/Source/PropertyAnimatorCore/Private/Properties/Converters/PropertyAnimatorCoreTextToStringConverter.h b/Engine/Plugins/Experimental/PropertyAnimatorCore/Source/PropertyAnimatorCore/Private/Properties/Converters/PropertyAnimatorCoreTextToStringConverter.h
--- a/Engine/Plugins/Experimental/PropertyAnimatorCore/Source/PropertyAnimatorCore/Private/Properties/Converters/PropertyAnimatorCoreTextToStringConverter.h
+++ b/Engine/Plugins/Experimental/PropertyAnimatorCore/Source/PropertyAnimatorCore/Private/Properties/Converters/PropertyAnimatorCoreTextToStringConverter.h
@@ -13,7 +13,7 @@ class UPropertyAnimatorCoreTextToStringConverter : public UPropertyAnimatorCoreC
 public:
 	//~ Begin UPropertyAnimatorCoreConverterBase
 	virtual bool IsConversionSupported(const FPropertyBagPropertyDesc& InFromProperty, const FPropertyBagPropertyDesc& InToProperty) const override;
-	virtual bool Convert(const FPropertyBagPropertyDesc& InFromProperty, const FInstancedPropertyBag& InFromBag, const FPropertyBagPropertyDesc& InToProperty, FInstancedPropertyBag& InToBag, const FInstancedStruct* InRule) override;
+	virtual bool Convert(const FPropertyBagPropertyDesc& InFromProperty, const FInstancedPropertyBag& InFromBag, const FPropertyBagPropertyDesc& InToProperty, FInstancedPropertyBag& InToBag, const TInstancedStruct<FPropertyAnimatorCoreConverterRuleBase>& InRule) override;
 	virtual UScriptStruct* GetConversionRuleStruct() const override;
 	//~ End UPropertyAnimatorCoreConverterBase
 };
\ No newline at end of file
diff --git a/Engine/Plugins/Experimental/PropertyAnimatorCore/Source/PropertyAnimatorCore/Private/Properties/PropertyAnimatorCoreContext.cpp b/Engine/Plugins/Experimental/PropertyAnimatorCore/Source/PropertyAnimatorCore/Private/Properties/PropertyAnimatorCoreContext.cpp
--- a/Engine/Plugins/Experimental/PropertyAnimatorCore/Source/PropertyAnimatorCore/Private/Properties/PropertyAnimatorCoreContext.cpp
+++ b/Engine/Plugins/Experimental/PropertyAnimatorCore/Source/PropertyAnimatorCore/Private/Properties/PropertyAnimatorCoreContext.cpp
@@ -17,7 +17,7 @@ TArray<FPropertyAnimatorCoreData> UPropertyAnimatorCoreContext::ResolveProperty(
 
 	if (UPropertyAnimatorCoreResolver* PropertyResolver = GetResolver())
 	{
-		PropertyResolver->ResolveProperties(AnimatedProperty, ResolvedProperties, bInForEvaluation);
+		PropertyResolver->ResolveTemplateProperties(AnimatedProperty, ResolvedProperties, bInForEvaluation);
 	}
 	else
 	{
@@ -110,9 +110,9 @@ void UPropertyAnimatorCoreContext::SetConverterClass(TSubclassOf<UPropertyAnimat
 
 	if (const UPropertyAnimatorCoreConverterBase* Converter = InConverterClass.GetDefaultObject())
 	{
-		if (UScriptStruct* RuleStruct = Converter->GetConversionRuleStruct())
+		if (const UScriptStruct* RuleStruct = Converter->GetConversionRuleStruct())
 		{
-			ConverterRule = FInstancedStruct(RuleStruct);
+			ConverterRule.InitializeAsScriptStruct(RuleStruct);
 			CheckEditConverterRule();
 		}
 	}
@@ -276,7 +276,7 @@ bool UPropertyAnimatorCoreContext::ExportPreset(const UPropertyAnimatorCorePrese
 
 void UPropertyAnimatorCoreContext::OnAnimatedPropertyLinked()
 {
-	bEditMagnitude = AnimatedProperty.IsA<FNumericProperty>() || AnimatedProperty.HasA<FNumericProperty>();
+	bEditMagnitude = AnimatedProperty.IsResolved();
 
 	if (const UPropertyAnimatorCoreResolver* PropertyResolver = AnimatedProperty.GetPropertyResolver())
 	{
@@ -490,7 +490,7 @@ void UPropertyAnimatorCoreContext::CommitEvaluationResult(const FPropertyAnimato
 
 	if (UPropertyAnimatorCoreConverterBase* Converter = ConverterClass.GetDefaultObject())
 	{
-		if (!Converter->Convert(*FromProperty, InEvaluatedValues, *ToProperty, DeltaPropertyValues, ConverterRule.IsValid() ? &ConverterRule : nullptr))
+		if (!Converter->Convert(*FromProperty, InEvaluatedValues, *ToProperty, DeltaPropertyValues, ConverterRule))
 		{
 			return;
 		}
diff --git a/Engine/Plugins/Experimental/PropertyAnimatorCore/Source/PropertyAnimatorCore/Public/Properties/Converters/PropertyAnimatorCoreConverterBase.h b/Engine/Plugins/Experimental/PropertyAnimatorCore/Source/PropertyAnimatorCore/Public/Properties/Converters/PropertyAnimatorCoreConverterBase.h
--- a/Engine/Plugins/Experimental/PropertyAnimatorCore/Source/PropertyAnimatorCore/Public/Properties/Converters/PropertyAnimatorCoreConverterBase.h
+++ b/Engine/Plugins/Experimental/PropertyAnimatorCore/Source/PropertyAnimatorCore/Public/Properties/Converters/PropertyAnimatorCoreConverterBase.h
@@ -2,12 +2,13 @@
 
 #pragma once
 
+#include "StructUtils/InstancedStruct.h"
 #include "StructUtils/PropertyBag.h"
 #include "UObject/Object.h"
 #include "PropertyAnimatorCoreConverterBase.generated.h"
 
 class UScriptStruct;
-struct FInstancedStruct;
+struct FPropertyAnimatorCoreConverterRuleBase;
 
 /**
  * Abstract base class,
@@ -27,7 +28,7 @@ public:
 	}
 
 	/** Converts from input bag to output bag with optional rules, returns true on success */
-	virtual bool Convert(const FPropertyBagPropertyDesc& InFromProperty, const FInstancedPropertyBag& InFromBag, const FPropertyBagPropertyDesc& InToProperty, FInstancedPropertyBag& InToBag, const FInstancedStruct* InRule = nullptr)
+	virtual bool Convert(const FPropertyBagPropertyDesc& InFromProperty, const FInstancedPropertyBag& InFromBag, const FPropertyBagPropertyDesc& InToProperty, FInstancedPropertyBag& InToBag, const TInstancedStruct<FPropertyAnimatorCoreConverterRuleBase>& InRule)
 	{
 		return false;
 	}
diff --git a/Engine/Plugins/Experimental/PropertyAnimatorCore/Source/PropertyAnimatorCore/Public/Properties/Converters/PropertyAnimatorCoreConverterTraits.h b/Engine/Plugins/Experimental/PropertyAnimatorCore/Source/PropertyAnimatorCore/Public/Properties/Converters/PropertyAnimatorCoreConverterTraits.h
--- a/Engine/Plugins/Experimental/PropertyAnimatorCore/Source/PropertyAnimatorCore/Public/Properties/Converters/PropertyAnimatorCoreConverterTraits.h
+++ b/Engine/Plugins/Experimental/PropertyAnimatorCore/Source/PropertyAnimatorCore/Public/Properties/Converters/PropertyAnimatorCoreConverterTraits.h
@@ -4,6 +4,13 @@
 
 #include "PropertyAnimatorCoreConverterTraits.generated.h"
 
+/** Base struct for converter rules */
+USTRUCT()
+struct FPropertyAnimatorCoreConverterRuleBase
+{
+	GENERATED_BODY()
+};
+
 /** Helper to define a converter from a type to another along an optional struct for rules */
 template<typename InFromType, typename InToType, typename InRuleType = void>
 struct TValueConverterTraits
@@ -52,7 +59,7 @@ struct FBoolConverterCondition
 };
 
 USTRUCT(BlueprintType)
-struct FBoolConverterRule
+struct FBoolConverterRule : public FPropertyAnimatorCoreConverterRuleBase
 {
 	GENERATED_BODY()
 
@@ -235,7 +242,7 @@ enum class EInt32ConverterMethod : uint8
 };
 
 USTRUCT(BlueprintType)
-struct FInt32ConverterRule
+struct FInt32ConverterRule : public FPropertyAnimatorCoreConverterRuleBase
 {
 	GENERATED_BODY()
 
diff --git a/Engine/Plugins/Experimental/PropertyAnimatorCore/Source/PropertyAnimatorCore/Public/Properties/PropertyAnimatorCoreContext.h b/Engine/Plugins/Experimental/PropertyAnimatorCore/Source/PropertyAnimatorCore/Public/Properties/PropertyAnimatorCoreContext.h
--- a/Engine/Plugins/Experimental/PropertyAnimatorCore/Source/PropertyAnimatorCore/Public/Properties/PropertyAnimatorCoreContext.h
+++ b/Engine/Plugins/Experimental/PropertyAnimatorCore/Source/PropertyAnimatorCore/Public/Properties/PropertyAnimatorCoreContext.h
@@ -2,6 +2,7 @@
 
 #pragma once
 
+#include "Converters/PropertyAnimatorCoreConverterTraits.h"
 #include "PropertyAnimatorCoreData.h"
 #include "Presets/PropertyAnimatorCorePresetable.h"
 #include "Presets/PropertyAnimatorCorePresetBase.h"
@@ -9,11 +10,11 @@
 #include "UObject/Object.h"
 #include "PropertyAnimatorCoreContext.generated.h"
 
-struct FInstancedStruct;
 class UPropertyAnimatorCoreBase;
 class UPropertyAnimatorCoreConverterBase;
 class UPropertyAnimatorCoreResolver;
 class UScriptStruct;
+struct FPropertyAnimatorCoreConverterRuleBase;
 
 /** Mode supported for properties value */
 UENUM(BlueprintType)
@@ -174,8 +175,8 @@ private:
 	bool bEditConverterRule = false;
 
 	/** If a converter is used, rules may be used to convert the property */
-	UPROPERTY(EditInstanceOnly, Category="Animator", meta=(HideEditConditionToggle, EditCondition="bEditConverterRule", EditConditionHides, AllowPrivateAccess="true"))
-	FInstancedStruct ConverterRule;
+	UPROPERTY(EditInstanceOnly, Category="Animator", NoClear, meta=(HideEditConditionToggle, EditCondition="bEditConverterRule", EditConditionHides, AllowPrivateAccess="true"))
+	TInstancedStruct<FPropertyAnimatorCoreConverterRuleBase> ConverterRule;
 
 	/** Edit condition for property resolver */
 	UPROPERTY(Transient)
diff --git a/Engine/Plugins/Experimental/PropertyAnimatorCore/Source/PropertyAnimatorCore/Public/Properties/PropertyAnimatorCoreResolver.h b/Engine/Plugins/Experimental/PropertyAnimatorCore/Source/PropertyAnimatorCore/Public/Properties/PropertyAnimatorCoreResolver.h
--- a/Engine/Plugins/Experimental/PropertyAnimatorCore/Source/PropertyAnimatorCore/Public/Properties/PropertyAnimatorCoreResolver.h
+++ b/Engine/Plugins/Experimental/PropertyAnimatorCore/Source/PropertyAnimatorCore/Public/Properties/PropertyAnimatorCoreResolver.h
@@ -37,7 +37,7 @@ public:
 	virtual void GetTemplateProperties(UObject* InContext, TSet<FPropertyAnimatorCoreData>& OutProperties, const TArray<FName>* InSearchPath = nullptr) {}
 
 	/** Called when we actually need the underlying properties of a template property */
-	virtual void ResolveProperties(const FPropertyAnimatorCoreData& InTemplateProperty, TArray<FPropertyAnimatorCoreData>& OutProperties, bool bInForEvaluation) {}
+	virtual void ResolveTemplateProperties(const FPropertyAnimatorCoreData& InTemplateProperty, TArray<FPropertyAnimatorCoreData>& OutProperties, bool bInForEvaluation) {}
 
 	FName GetResolverName() const
 	{
