MotionDesign : PropertyAnimator
- Fixed numeric animator would wait for all range elements to finish current animation cycle before entering new cycle creating a gap
- Fixed animator negative offset on range property does not reverse order of animation

diff --git a/Engine/Plugins/VirtualProduction/PropertyAnimator/Source/PropertyAnimator/Private/Animators/PropertyAnimatorNumericBase.cpp b/Engine/Plugins/VirtualProduction/PropertyAnimator/Source/PropertyAnimator/Private/Animators/PropertyAnimatorNumericBase.cpp
--- a/Engine/Plugins/VirtualProduction/PropertyAnimator/Source/PropertyAnimator/Private/Animators/PropertyAnimatorNumericBase.cpp
+++ b/Engine/Plugins/VirtualProduction/PropertyAnimator/Source/PropertyAnimator/Private/Animators/PropertyAnimatorNumericBase.cpp
@@ -164,52 +164,97 @@ void UPropertyAnimatorNumericBase::EvaluateProperties(FInstancedPropertyBag& InP
 		const int32 RangeIndex = InRangeIndex + 1;
 		const int32 RangeMax = InRangeMax + 1;
 		const double TimeOffset = InOptions->GetTimeOffset() / RangeMax;
-		const double AbsTimeOffset = FMath::Abs(TimeOffset);
-		const double MaxTimeOffset = RangeMax * AbsTimeOffset;
 		double PropertyTimeElapsed = TimeElapsed + RandomTimeOffset;
 		double Frequency = CycleDuration != 0.f ? 1.f / CycleDuration : 0.f;
 
-		PropertyTimeElapsed += RangeIndex * TimeOffset;
-		PropertyTimeElapsed = FMath::Abs(PropertyTimeElapsed);
-
-		if (CycleMode == EPropertyAnimatorCycleMode::DoOnce)
+		// handle time offset for range
+		if (TimeOffset >= 0.0)
 		{
-			if (PropertyTimeElapsed > CycleDuration)
-			{
-				PropertyTimeElapsed = CycleDuration - UE_KINDA_SMALL_NUMBER;
-			}
+		    // positive offset: last element starts first
+		    PropertyTimeElapsed += RangeIndex * TimeOffset;
 		}
-		else if (CycleMode == EPropertyAnimatorCycleMode::Loop)
+		else
 		{
-			PropertyTimeElapsed = FMath::Fmod(PropertyTimeElapsed, CycleDuration + MaxTimeOffset + CycleGapDuration);
-
-			if (PropertyTimeElapsed > CycleDuration)
-			{
-				PropertyTimeElapsed = CycleDuration - UE_KINDA_SMALL_NUMBER;
-			}
+		    // negative offset: first element starts first
+		    PropertyTimeElapsed += (InRangeMax - InRangeIndex) * FMath::Abs(TimeOffset);
 		}
-		else if (CycleMode == EPropertyAnimatorCycleMode::PingPong)
+
+		// wrap for do once
+		if (CycleMode == EPropertyAnimatorCycleMode::DoOnce)
 		{
-			const bool bReverse = FMath::Modulo(FMath::TruncToInt32(PropertyTimeElapsed / (CycleDuration + MaxTimeOffset + CycleGapDuration)), 2) != 0;
-			PropertyTimeElapsed = FMath::Fmod(PropertyTimeElapsed, CycleDuration + MaxTimeOffset + CycleGapDuration);
-
-			if (PropertyTimeElapsed > CycleDuration)
-			{
-				PropertyTimeElapsed = CycleDuration - UE_KINDA_SMALL_NUMBER;
-			}
-
-			if (bReverse)
-			{
-				PropertyTimeElapsed = CycleDuration - FMath::Fmod(PropertyTimeElapsed, CycleDuration);
-			}
-			else
-			{
-				PropertyTimeElapsed = FMath::Fmod(PropertyTimeElapsed, CycleDuration);
-			}
+		    if (PropertyTimeElapsed < 0.0)
+		    {
+		        PropertyTimeElapsed = 0.0;
+		    }
+		    else if (PropertyTimeElapsed > CycleDuration)
+		    {
+		        PropertyTimeElapsed = CycleDuration - UE_KINDA_SMALL_NUMBER;
+		    }
 		}
-		else if (CycleMode == EPropertyAnimatorCycleMode::None)
+		else
 		{
-			Frequency = 1;
+			// wrap negative time
+		    if (PropertyTimeElapsed < 0.0)
+		    {
+		        double Period = CycleDuration + CycleGapDuration;
+		        if (CycleMode == EPropertyAnimatorCycleMode::PingPong)
+		        {
+		            Period = 2.0 * CycleDuration + 2.0 * CycleGapDuration;
+		        }
+
+		        PropertyTimeElapsed = FMath::Fmod(PropertyTimeElapsed, Period);
+		        if (PropertyTimeElapsed < 0.0)
+		        {
+		            PropertyTimeElapsed += Period;
+		        }
+		    }
+
+		    if (CycleMode == EPropertyAnimatorCycleMode::Loop)
+		    {
+		        const double Period = CycleDuration + CycleGapDuration - UE_KINDA_SMALL_NUMBER;
+		        const double TimeInPeriod = FMath::Fmod(PropertyTimeElapsed, Period);
+
+		        if (TimeInPeriod <= CycleDuration)
+		        {
+		        	// active phase
+		            PropertyTimeElapsed = TimeInPeriod;
+		        }
+		        else
+		        {
+		        	// gap phase
+		            PropertyTimeElapsed = CycleDuration - UE_KINDA_SMALL_NUMBER;
+		        }
+		    }
+		    else if (CycleMode == EPropertyAnimatorCycleMode::PingPong)
+		    {
+		        const double Period = 2.0 * CycleDuration + 2.0 * CycleGapDuration - UE_KINDA_SMALL_NUMBER;
+		        const double TimeInPeriod = FMath::Fmod(PropertyTimeElapsed, Period);
+
+		        if (TimeInPeriod < CycleDuration)
+		        {
+		        	// forward phase
+		            PropertyTimeElapsed = TimeInPeriod;
+		        }
+		        else if (TimeInPeriod < CycleDuration + CycleGapDuration)
+		        {
+		        	// forward gap hold
+		            PropertyTimeElapsed = CycleDuration - UE_KINDA_SMALL_NUMBER;
+		        }
+		        else if (TimeInPeriod < 2.0 * CycleDuration + CycleGapDuration)
+		        {
+		        	// backward phase
+		            PropertyTimeElapsed = 2.0 * CycleDuration + CycleGapDuration - TimeInPeriod;
+		        }
+		        else
+		        {
+		        	// backward gap hold
+		            PropertyTimeElapsed = 0.0;
+		        }
+		    }
+		    else if (CycleMode == EPropertyAnimatorCycleMode::None)
+		    {
+		        Frequency = 1;
+		    }
 		}
 
 		if (Magnitude != 0
diff --git a/Engine/Plugins/VirtualProduction/PropertyAnimatorCore/Source/PropertyAnimatorCore/Public/Animators/PropertyAnimatorCoreBase.h b/Engine/Plugins/VirtualProduction/PropertyAnimatorCore/Source/PropertyAnimatorCore/Public/Animators/PropertyAnimatorCoreBase.h
--- a/Engine/Plugins/VirtualProduction/PropertyAnimatorCore/Source/PropertyAnimatorCore/Public/Animators/PropertyAnimatorCoreBase.h
+++ b/Engine/Plugins/VirtualProduction/PropertyAnimatorCore/Source/PropertyAnimatorCore/Public/Animators/PropertyAnimatorCoreBase.h
@@ -375,7 +375,7 @@ private:
 	bool bOverrideTimeSource = true;
 
 	/** The time source to use */
-	UPROPERTY(EditInstanceOnly, Setter, Getter, Category="Animator", meta=(GetOptions="GetTimeSourceNames", EditCondition="bOverrideTimeSource"))
+	UPROPERTY(EditInstanceOnly, Setter, Getter, Category="Animator", meta=(GetOptions="GetTimeSourceNames", EditCondition="bOverrideTimeSource", NoResetToDefault))
 	FName TimeSourceName = NAME_None;
 
 	/** Active time source with its options, determined by its name */
diff --git a/Engine/Plugins/VirtualProduction/PropertyAnimatorCore/Source/PropertyAnimatorCore/Public/Components/PropertyAnimatorCoreComponent.h b/Engine/Plugins/VirtualProduction/PropertyAnimatorCore/Source/PropertyAnimatorCore/Public/Components/PropertyAnimatorCoreComponent.h
--- a/Engine/Plugins/VirtualProduction/PropertyAnimatorCore/Source/PropertyAnimatorCore/Public/Components/PropertyAnimatorCoreComponent.h
+++ b/Engine/Plugins/VirtualProduction/PropertyAnimatorCore/Source/PropertyAnimatorCore/Public/Components/PropertyAnimatorCoreComponent.h
@@ -139,7 +139,7 @@ protected:
 	float AnimatorsMagnitude = 1.f;
 
 	/** The global time source to use, can be overriden in animator */
-	UPROPERTY(EditInstanceOnly, Setter, Getter, DisplayName="Global Time Source Name", Category="Animator", meta=(GetOptions="GetTimeSourceNames"))
+	UPROPERTY(EditInstanceOnly, Setter, Getter, DisplayName="Global Time Source Name", Category="Animator", meta=(GetOptions="GetTimeSourceNames", NoResetToDefault))
 	FName AnimatorsTimeSourceName = NAME_None;
 
 	/** Active time source with its options, determined by its name */
diff --git a/Engine/Plugins/VirtualProduction/PropertyAnimatorCore/Source/PropertyAnimatorCoreEditor/Private/Sequencer/MovieSceneAnimatorTrackEditor.cpp b/Engine/Plugins/VirtualProduction/PropertyAnimatorCore/Source/PropertyAnimatorCoreEditor/Private/Sequencer/MovieSceneAnimatorTrackEditor.cpp
--- a/Engine/Plugins/VirtualProduction/PropertyAnimatorCore/Source/PropertyAnimatorCoreEditor/Private/Sequencer/MovieSceneAnimatorTrackEditor.cpp
+++ b/Engine/Plugins/VirtualProduction/PropertyAnimatorCore/Source/PropertyAnimatorCoreEditor/Private/Sequencer/MovieSceneAnimatorTrackEditor.cpp
@@ -121,10 +121,10 @@ void FMovieSceneAnimatorTrackEditor::ExecuteAddTrack(const TArray<UObject*>& InO
 
 		const FGuid ObjectBinding = SequencerPtr->GetHandleToObject(Owner, /** Create */true);
 
-		if (FocusedMovieScene->FindSpawnable(ObjectBinding))
+		if (!ObjectBinding.IsValid() || FocusedMovieScene->FindSpawnable(ObjectBinding))
 		{
 			// we only want to add tracks for possessables
-			return;
+			continue;
 		}
 
 		if (!FocusedMovieScene->FindTrack<UMovieSceneAnimatorTrack>(ObjectBinding))
