MotionDesign : PropertyAnimator
- Fixed clock animator with missing format and padding

diff --git a/Engine/Plugins/VirtualProduction/PropertyAnimator/Source/PropertyAnimator/Private/Animators/PropertyAnimatorClock.cpp b/Engine/Plugins/VirtualProduction/PropertyAnimator/Source/PropertyAnimator/Private/Animators/PropertyAnimatorClock.cpp
--- a/Engine/Plugins/VirtualProduction/PropertyAnimator/Source/PropertyAnimator/Private/Animators/PropertyAnimatorClock.cpp
+++ b/Engine/Plugins/VirtualProduction/PropertyAnimator/Source/PropertyAnimator/Private/Animators/PropertyAnimatorClock.cpp
@@ -45,10 +45,13 @@ UPropertyAnimatorClock::UPropertyAnimatorClock()
 		RegisterFormat(TEXT('I'), [](const FDateTime& InDateTime)->FString{ return InDateTime.ToFormattedString(TEXT("%I")); });
 		RegisterFormat(TEXT('H'), [](const FDateTime& InDateTime)->FString{ return InDateTime.ToFormattedString(TEXT("%H")); });
 		RegisterFormat(TEXT('h'), [](const FDateTime& InDateTime)->FString{ return FString::FromInt(InDateTime.GetHour()); });
+		RegisterFormat(TEXT('o'), [](const FDateTime& InDateTime)->FString{ return LexToString(InDateTime.GetTicks() / ETimespan::TicksPerHour); });
 		RegisterFormat(TEXT('M'), [](const FDateTime& InDateTime)->FString{ return InDateTime.ToFormattedString(TEXT("%M")); });
 		RegisterFormat(TEXT('N'), [](const FDateTime& InDateTime)->FString{ return FString::FromInt(InDateTime.GetMinute()); });
+		RegisterFormat(TEXT('i'), [](const FDateTime& InDateTime)->FString{ return LexToString(InDateTime.GetTicks() / ETimespan::TicksPerMinute); });
 		RegisterFormat(TEXT('S'), [](const FDateTime& InDateTime)->FString{ return InDateTime.ToFormattedString(TEXT("%S")); });
 		RegisterFormat(TEXT('s'), [](const FDateTime& InDateTime)->FString{ return FString::FromInt(InDateTime.GetSecond()); });
+		RegisterFormat(TEXT('c'), [](const FDateTime& InDateTime)->FString{ return LexToString(InDateTime.GetTicks() / ETimespan::TicksPerSecond); });
 		RegisterFormat(TEXT('f'), [](const FDateTime& InDateTime)->FString{ return InDateTime.ToString(TEXT("%s")); });
 		RegisterFormat(TEXT('F'), [](const FDateTime& InDateTime)->FString{ return FString::FromInt(InDateTime.GetMillisecond()); });
 		RegisterFormat(TEXT('p'), [](const FDateTime& InDateTime)->FString{ return InDateTime.ToFormattedString(TEXT("%p")); });
@@ -62,10 +65,14 @@ void UPropertyAnimatorClock::SetDisplayFormat(const FString& InDisplayFormat)
 	DisplayFormat = InDisplayFormat;
 }
 
-FString UPropertyAnimatorClock::FormatDateTime(const FDateTime& InDateTime, const FString& InDisplayFormat)
+void UPropertyAnimatorClock::SetDisplayMask(const FString& InPadding)
 {
-	FString Result;
-	Result.Reserve(InDisplayFormat.Len());
+	DisplayMask = InPadding;
+}
+
+FString UPropertyAnimatorClock::FormatDateTime(const FDateTime& InDateTime, const FString& InDisplayFormat, const FString& InMask)
+{
+	FStringBuilderBase Builder;
 
     for (int32 CharIndex = 0; CharIndex < InDisplayFormat.Len(); ++CharIndex)
     {
@@ -73,18 +80,23 @@ FString UPropertyAnimatorClock::FormatDateTime(const FDateTime& InDateTime, cons
         {
 			if (const TFunction<FString(const FDateTime&)>* Formatter = UE::Private::Animator::FormatFunctions.Find(InDisplayFormat[CharIndex + 1]))
 			{
-				Result += (*Formatter)(InDateTime);
+				Builder.Append((*Formatter)(InDateTime));
 			}
 
             ++CharIndex;
         }
         else
         {
-            Result += InDisplayFormat[CharIndex];
+        	Builder.AppendChar(InDisplayFormat[CharIndex]);
         }
     }
 
-    return Result;
+	if (Builder.Len() < InMask.Len())
+	{
+		Builder.Prepend(InMask.Left(InMask.Len() - Builder.Len()));
+	}
+
+    return *Builder;
 }
 
 void UPropertyAnimatorClock::OnAnimatorRegistered(FPropertyAnimatorCoreMetadata& InMetadata)
@@ -101,7 +113,7 @@ void UPropertyAnimatorClock::EvaluateProperties(FInstancedPropertyBag& InParamet
 
 	const FTimespan ElapsedTimeSpan = FTimespan::FromSeconds(TimeElapsed);
 	const FDateTime DateTime(ElapsedTimeSpan > FTimespan::Zero() ? ElapsedTimeSpan.GetTicks() : 0);
-	const FString FormattedDateTime = FormatDateTime(DateTime, DisplayFormat);
+	const FString FormattedDateTime = FormatDateTime(DateTime, DisplayFormat, DisplayMask);
 
 	EvaluateEachLinkedProperty([this, FormattedDateTime](
 		UPropertyAnimatorCoreContext* InContext
@@ -129,6 +141,10 @@ bool UPropertyAnimatorClock::ImportPreset(const UPropertyAnimatorCorePresetBase*
 		AnimatorArchive->Get(GET_MEMBER_NAME_STRING_CHECKED(UPropertyAnimatorClock, DisplayFormat), DisplayFormatValue);
 		SetDisplayFormat(DisplayFormatValue);
 
+		FString DisplayMaskValue = DisplayMask;
+		AnimatorArchive->Get(GET_MEMBER_NAME_STRING_CHECKED(UPropertyAnimatorClock, DisplayMask), DisplayMaskValue);
+		SetDisplayMask(DisplayMaskValue);
+
 		return true;
 	}
 
@@ -142,6 +158,7 @@ bool UPropertyAnimatorClock::ExportPreset(const UPropertyAnimatorCorePresetBase*
 		const TSharedPtr<FPropertyAnimatorCorePresetObjectArchive> AnimatorArchive = OutValue->AsMutableObject();
 
 		AnimatorArchive->Set(GET_MEMBER_NAME_STRING_CHECKED(UPropertyAnimatorClock, DisplayFormat), DisplayFormat);
+		AnimatorArchive->Set(GET_MEMBER_NAME_STRING_CHECKED(UPropertyAnimatorClock, DisplayMask), DisplayMask);
 
 		return true;
 	}
diff --git a/Engine/Plugins/VirtualProduction/PropertyAnimator/Source/PropertyAnimator/Public/Animators/PropertyAnimatorClock.h b/Engine/Plugins/VirtualProduction/PropertyAnimator/Source/PropertyAnimator/Public/Animators/PropertyAnimatorClock.h
--- a/Engine/Plugins/VirtualProduction/PropertyAnimator/Source/PropertyAnimator/Public/Animators/PropertyAnimatorClock.h
+++ b/Engine/Plugins/VirtualProduction/PropertyAnimator/Source/PropertyAnimator/Public/Animators/PropertyAnimatorClock.h
@@ -28,7 +28,7 @@ class UPropertyAnimatorClock : public UPropertyAnimatorTextBase
 public:
 	static void RegisterFormat(TCHAR InChar, TFunction<FString(const FDateTime&)> InFormatter);
 	static void UnregisterFormat(TCHAR InChar);
-	static FString FormatDateTime(const FDateTime& InDateTime, const FString& InDisplayFormat);
+	static FString FormatDateTime(const FDateTime& InDateTime, const FString& InDisplayFormat, const FString& InMask);
 
 	UPropertyAnimatorClock();
 
@@ -38,6 +38,12 @@ public:
 		return DisplayFormat;
 	}
 
+	PROPERTYANIMATOR_API void SetDisplayMask(const FString& InPadding);
+	const FString& GetDisplayMask() const
+	{
+		return DisplayMask;
+	}
+
 protected:
 	//~ Begin UPropertyAnimatorCoreBase
 	virtual void OnAnimatorRegistered(FPropertyAnimatorCoreMetadata& InMetadata) override;
@@ -65,16 +71,26 @@ protected:
 	 * %I - 12h Hour, 01-12
 	 * %H - 24h Hour, 00-23
 	 * %h - 24h Hour, 0-23
+	 * %o - Hour, 0-inf
 	 * %M - Minute, 00-59
 	 * %N - Minute, 0-59
+	 * %i - Minute, 0-inf
 	 * %S - Second, 00-60
 	 * %s - Second, 0-60
+	 * %c - Second, 0-inf
 	 * %f - Millisecond, 000-999
 	 * %F - Millisecond, 0-999
 	 * %p - AM or PM
 	 * %P - am or PM
 	 * %t - Ticks since midnight, January 1, 0001
 	 */
-	UPROPERTY(EditInstanceOnly, BlueprintReadWrite, Setter, Getter, Category="Animator")
+	UPROPERTY(EditInstanceOnly, Setter, Getter, Category="Animator")
 	FString DisplayFormat = TEXT("%H:%M:%S");
+
+	/**
+	 * Used as padding to complete the display format
+	 * eg: DisplayMask=****** and DisplayFormat=%c would result in ***999
+	 */
+	UPROPERTY(EditInstanceOnly, Setter, Getter, Category="Animator")
+	FString DisplayMask = TEXT("");
 };
\ No newline at end of file
